ModelingMode: Fix issue with localization namespace being not undefined at end of file

diff --git a/Engine/Plugins/Experimental/ToolPresets/Source/PresetAsset/Private/PresetAssetModule.cpp b/Engine/Plugins/Experimental/ToolPresets/Source/PresetAsset/Private/PresetAssetModule.cpp
--- a/Engine/Plugins/Experimental/ToolPresets/Source/PresetAsset/Private/PresetAssetModule.cpp
+++ b/Engine/Plugins/Experimental/ToolPresets/Source/PresetAsset/Private/PresetAssetModule.cpp
@@ -61,4 +61,6 @@ void FPresetAssetModule::OnPostEngineInit()
 }
 
 
-IMPLEMENT_MODULE(FPresetAssetModule, PresetAsset);
\ No newline at end of file
+IMPLEMENT_MODULE(FPresetAssetModule, PresetAsset);
+
+#undef LOCTEXT_NAMESPACE
\ No newline at end of file
