Modeling Mode: Interactive Tool Preset Feature Version 2

Adds a preset manager with controls to edit preset labels and tooltips, display both project level and user level presets, and create/delete user level preset collections.
Presets are now saving internal tool setting state via json encoded strings, to avoid issues with transient class flags.
Presets are now either created by the user under project trees or via the preset manager in the plugin tree to act as user level preset collections.
A default preset collection is automatically created upon editor start up if none exists prior.
The default preset is used as the destination for all new presets until explicitly overridden by the user within an editor session.
Preset updating is removed for the time being.

diff --git a/Engine/Plugins/Editor/ModelingToolsEditorMode/Source/ModelingToolsEditorMode/ModelingToolsEditorMode.Build.cs b/Engine/Plugins/Editor/ModelingToolsEditorMode/Source/ModelingToolsEditorMode/ModelingToolsEditorMode.Build.cs
--- a/Engine/Plugins/Editor/ModelingToolsEditorMode/Source/ModelingToolsEditorMode/ModelingToolsEditorMode.Build.cs
+++ b/Engine/Plugins/Editor/ModelingToolsEditorMode/Source/ModelingToolsEditorMode/ModelingToolsEditorMode.Build.cs
@@ -52,7 +52,9 @@ public class ModelingToolsEditorMode : ModuleRules
 				"PropertyEditor",
 				"ToolMenus",
 				"EditorConfig",
-				"PresetAsset"
+				"PresetAsset",
+				"PresetEditor",
+				"EditorConfig",
 				// ... add private dependencies that you statically link with here ...	
 			}
 		);
diff --git a/Engine/Plugins/Editor/ModelingToolsEditorMode/Source/ModelingToolsEditorMode/Private/ModelingToolsEditorModeToolkit.cpp b/Engine/Plugins/Editor/ModelingToolsEditorMode/Source/ModelingToolsEditorMode/Private/ModelingToolsEditorModeToolkit.cpp
--- a/Engine/Plugins/Editor/ModelingToolsEditorMode/Source/ModelingToolsEditorMode/Private/ModelingToolsEditorModeToolkit.cpp
+++ b/Engine/Plugins/Editor/ModelingToolsEditorMode/Source/ModelingToolsEditorMode/Private/ModelingToolsEditorModeToolkit.cpp
@@ -51,6 +51,13 @@
 #include "Widgets/Input/STextComboBox.h"
 #include "ToolMenus.h"
 #include "ToolkitBuilder.h"
+#include "Dialogs/Dialogs.h"
+#include "Widgets/Input/SComboButton.h"
+#include "IPresetEditorModule.h"
+#include "PresetSettings.h"
+#include "AssetRegistry/AssetRegistryModule.h"
+#include "ModelingWidgets/SToolInputAssetComboPanel.h"
+#include "Fonts/SlateFontInfo.h"
 
 
 #define LOCTEXT_NAMESPACE "FModelingToolsEditorModeToolkit"
@@ -67,14 +74,14 @@ static TAutoConsoleVariable<int32> CVarEnablePolyModeling(
 	TEXT("Enable prototype PolyEdit tab"));
 static TAutoConsoleVariable<int32> CVarEnableToolPresets(
 	TEXT("modeling.EnablePresets"),
-	0,
+	1,
 	TEXT("Enable tool preset features and UX"));
 
 namespace FModelingToolsEditorModeToolkitLocals
 {
 	typedef TFunction<void(UInteractiveToolsPresetCollectionAsset& Preset, UInteractiveTool& Tool)> PresetAndToolFunc;
 
-	void ExecuteWithPresetAndTool(UEdMode& EdMode, EToolSide ToolSide, FSoftObjectPath& PresetPath, PresetAndToolFunc Function)
+	void ExecuteWithPresetAndTool(UEdMode& EdMode, EToolSide ToolSide, const FSoftObjectPath& PresetPath, PresetAndToolFunc Function)
 	{
 		UInteractiveToolsPresetCollectionAsset* Preset = nullptr;		
 		UInteractiveTool* Tool = EdMode.GetToolManager()->GetActiveTool(EToolSide::Left);
@@ -107,12 +114,8 @@ FModelingToolsEditorModeToolkit::FModelingToolsEditorModeToolkit()
 	UModelingModeEditableToolPaletteConfig::Initialize();
 	UModelingModeEditableToolPaletteConfig::Get()->LoadEditorConfig();
 
-	PresetSettings = NewObject<UPresetSettingsProperties>();
-}
-
-FString FModelingToolsEditorModeToolkit::GetReferencerName() const
-{
-	return TEXT("FModelingToolsEditorModeToolkit");
+	UPresetUserSettings::Initialize();
+	UPresetUserSettings::Get()->LoadEditorConfig();
 }
 
 FModelingToolsEditorModeToolkit::~FModelingToolsEditorModeToolkit()
@@ -135,8 +138,6 @@ FModelingToolsEditorModeToolkit::~FModelingToolsEditorModeToolkit()
 	Settings->OnModified.Remove(AssetSettingsModifiedHandle);
 	GetScriptableEditorMode()->GetInteractiveToolsContext(EToolsContextScope::EdMode)->OnToolNotificationMessage.RemoveAll(this);
 	GetScriptableEditorMode()->GetInteractiveToolsContext(EToolsContextScope::EdMode)->OnToolWarningMessage.RemoveAll(this);
-
-	PresetSettings = nullptr;
 }
 
 
@@ -460,10 +461,12 @@ void FModelingToolsEditorModeToolkit::Init(const TSharedPtr<IToolkitHost>& InitT
 		.ColorAndOpacity(FSlateColor(FLinearColor(0.9f, 0.15f, 0.15f)));
 	ToolWarningArea->SetText(FText::GetEmpty());
 
+	ToolPresetArea = MakePresetPanel();
+
 	// add the various sections to the mode toolbox
 	ToolkitWidgetVBox->AddSlot().AutoHeight().HAlign(HAlign_Fill).Padding(5)
 		[
-			MakePresetPanel()->AsShared()
+			ToolPresetArea->AsShared()
 		];
 	ToolkitWidgetVBox->AddSlot().AutoHeight().HAlign(HAlign_Fill).Padding(5)
 		[
@@ -500,6 +503,7 @@ void FModelingToolsEditorModeToolkit::Init(const TSharedPtr<IToolkitHost>& InitT
 	if (HasToolkitBuilder())
 	{
 		ToolkitSections->ModeWarningArea = ModeWarningArea;
+		ToolkitSections->ToolPresetArea = ToolPresetArea;
 		ToolkitSections->DetailsView = ModeDetailsView;
 		ToolkitSections->ToolWarningArea = ToolWarningArea;
 		ToolkitSections->Footer = MakeAssetConfigPanel();
@@ -535,14 +539,10 @@ void FModelingToolsEditorModeToolkit::Init(const TSharedPtr<IToolkitHost>& InitT
 		MakeSelectionPaletteOverlayWidget();
 		GetToolkitHost()->AddViewportOverlayWidget(SelectionPaletteOverlayWidget.ToSharedRef());
 	}
-}
 
-void FModelingToolsEditorModeToolkit::AddReferencedObjects(FReferenceCollector& Collector)
-{
-	Collector.AddReferencedObject(PresetSettings);
+	CurrentPreset = FSoftObjectPath(FTopLevelAssetPath("/ToolPresets/Presets/_DefaultCollection", "_DefaultCollection"));
 }
 
-
 void FModelingToolsEditorModeToolkit::MakeToolShutdownOverlayWidget()
 {
 	SAssignNew(ToolShutdownViewportOverlayWidget, SHorizontalBox)
@@ -846,7 +846,7 @@ TSharedPtr<SWidget> FModelingToolsEditorModeToolkit::MakePresetPanel()
 	const TSharedPtr<SVerticalBox> Content = SNew(SVerticalBox);
 	UModelingToolsEditorModeSettings* Settings = GetMutableDefault<UModelingToolsEditorModeSettings>();
 
-	bool bEnableToolPresets = (CVarEnableToolPresets.GetValueOnGameThread() > 0);
+	bool bEnableToolPresets = (CVarEnableToolPresets.GetValueOnGameThread() > 0);	
 	if (!bEnableToolPresets || Settings->InRestrictiveMode())
 	{
 		return SNew(SVerticalBox);
@@ -854,12 +854,15 @@ TSharedPtr<SWidget> FModelingToolsEditorModeToolkit::MakePresetPanel()
 
 	const TSharedPtr<SHorizontalBox> NewContent = SNew(SHorizontalBox);
 	
+	auto IsToolActive = [this]() {return this->OwningEditorMode->GetToolManager()->GetActiveTool(EToolSide::Left) != nullptr ? EVisibility::Visible : EVisibility::Collapsed; };
+
 	NewContent->AddSlot().HAlign(HAlign_Right)
 	[
 		SNew(SComboButton)
 		.ComboButtonStyle(&FAppStyle::Get().GetWidgetStyle<FComboButtonStyle>("SimpleComboButton"))
 		.OnGetMenuContent(this, &FModelingToolsEditorModeToolkit::GetPresetCreateButtonContent)
 		.HasDownArrow(true)
+		.Visibility_Lambda(IsToolActive)
 		.ButtonContent()
 		[
 			SNew(SHorizontalBox)
@@ -867,15 +870,6 @@ TSharedPtr<SWidget> FModelingToolsEditorModeToolkit::MakePresetPanel()
 			.AutoWidth()
 			.VAlign(VAlign_Center)
 			.Padding(4.0, 0.0f)
-			[
-				SNew(SImage)
-				.ColorAndOpacity(FSlateColor::UseForeground())
-				.Image(FAppStyle::Get().GetBrush("Icons.Plus"))
-			]
-			+ SHorizontalBox::Slot()
-			.AutoWidth()
-			.VAlign(VAlign_Center)
-			.Padding(4.0, 0.0f)
 			[
 				SNew(STextBlock)
 				.Text(LOCTEXT("ModelingPresetPanelHeader", "Presets"))
@@ -886,28 +880,6 @@ TSharedPtr<SWidget> FModelingToolsEditorModeToolkit::MakePresetPanel()
 		]
 	];
 
-	NewContent->AddSlot().HAlign(HAlign_Right).AutoWidth()
-		[
-			SNew(SComboButton)
-			.ComboButtonStyle(&FAppStyle::Get().GetWidgetStyle<FComboButtonStyle>("SimpleComboButton"))
-		.OnGetMenuContent(this, &FModelingToolsEditorModeToolkit::GetPresetSettingsButtonContent)
-		.OnMenuOpenChanged(this, &FModelingToolsEditorModeToolkit::RebuildPresetListForTool)
-		.HasDownArrow(false)
-		.ButtonContent()
-		[
-			SNew(SHorizontalBox)
-			+ SHorizontalBox::Slot()
-		.AutoWidth()
-		.VAlign(VAlign_Center)
-		.Padding(4.0, 0.0f)
-		[
-			SNew(SImage)
-			.ColorAndOpacity(FSlateColor::UseForeground())
-		.Image(FAppStyle::Get().GetBrush("Icons.Settings"))
-		]
-		]
-		];
-
 	TSharedPtr<SHorizontalBox> AssetConfigPanel = SNew(SHorizontalBox)
 		+ SHorizontalBox::Slot()
 		[
@@ -922,36 +894,113 @@ bool FModelingToolsEditorModeToolkit::IsPresetEnabled() const
 	return CurrentPreset.IsAsset();
 }
 
-TSharedRef<SWidget> FModelingToolsEditorModeToolkit::GetPresetSettingsButtonContent()
-{
-    const TSharedPtr<SVerticalBox> Content = SNew(SVerticalBox);
-
-	FPropertyEditorModule& PropertyEditorModule = FModuleManager::GetModuleChecked<FPropertyEditorModule>("PropertyEditor");
 
-	TSharedPtr<IDetailsView> PresetSettingsDetailsView;
+TSharedRef<SWidget> FModelingToolsEditorModeToolkit::GetPresetCreateButtonContent()
+{
+	RebuildPresetListForTool(false);
 
-	FDetailsViewArgs DetailsViewArgs;
-	DetailsViewArgs.bAllowSearch = false;
-	DetailsViewArgs.NameAreaSettings = FDetailsViewArgs::HideNameArea;
-	DetailsViewArgs.bHideSelectionTip = true;
-	DetailsViewArgs.DefaultsOnlyVisibility = EEditDefaultsOnlyNodeVisibility::Automatic;
-	DetailsViewArgs.bShowOptions = false;
-	DetailsViewArgs.bAllowMultipleTopLevelObjects = true;
+		auto OpenNewPresetDialog = [this]()
+	{
+		// Set the result if they just click Ok
+		SGenericDialogWidget::FArguments FolderDialogArguments;
+		FolderDialogArguments.OnOkPressed_Lambda([this]()
+			{
+				CreateNewPresetInCollection(NewPresetLabel,
+					CurrentPreset,
+					NewPresetTooltip,
+					NewPresetIcon);
+			});
 
-	PresetSettingsDetailsView = PropertyEditorModule.CreateDetailView(DetailsViewArgs);
-	
-	PresetSettingsDetailsView->SetObject(PresetSettings);
+		// Present the Dialog
+		SGenericDialogWidget::OpenDialog(LOCTEXT("ToolPresets_CreatePreset", "Create new preset from active tool's settings"),
+			SNew(SVerticalBox)
+			+ SVerticalBox::Slot()
+			.AutoHeight()
+			.Padding(5.0f)
+			[
+				SNew(SHorizontalBox)
+				+ SHorizontalBox::Slot()
+				[
+					SNew(STextBlock)
+					.Text(LOCTEXT("ToolPresets_CreatePresetLabel", "Label"))
+				]
+				+ SHorizontalBox::Slot()
+				[
+					SNew(SEditableTextBox)
+					.OnTextCommitted_Lambda([this](const FText& NewLabel, ETextCommit::Type&) { NewPresetLabel = NewLabel.ToString(); })
+					.OverflowPolicy(ETextOverflowPolicy::Ellipsis)
+					.ToolTipText(LOCTEXT("ToolPresets_CreatePresetLabel_Tooltip", "A short, descriptive identifier for the new preset."))
+				]
+			]
+		+ SVerticalBox::Slot()
+			.AutoHeight()
+			.Padding(5.0f)
+			[
+				SNew(SHorizontalBox)
+				+ SHorizontalBox::Slot()
+				[
+					SNew(STextBlock)
+					.Text(LOCTEXT("ToolPresets_CreatePresetTooltip", "Tooltip"))
+				]
+				+ SHorizontalBox::Slot()
+				[
+					SNew(SEditableTextBox)
+					.OnTextCommitted_Lambda([this](const FText& NewToolTip, ETextCommit::Type&) { NewPresetTooltip = NewToolTip.ToString(); })
+					.OverflowPolicy(ETextOverflowPolicy::Ellipsis)
+					.ToolTipText(LOCTEXT("ToolPresets_CreatePresetTooltip_Tooltip", "A descriptive tooltip for the new preset."))
+				]
+			]
+		+ SVerticalBox::Slot()
+			.AutoHeight()
+			.Padding(5.0f)
+			[
+				SNew(SExpandableArea)
+				.AreaTitle(FText::FromString("Advanced"))
+			    .InitiallyCollapsed(true)
+				.BodyContent()
+				[
 
-	Content->AddSlot()
-	[
-		PresetSettingsDetailsView->AsShared()
-	];
+					SNew(SHorizontalBox)
+					+ SHorizontalBox::Slot()
+					.AutoWidth()
+					[
+						SNew(SToolInputAssetComboPanel)
+						.AssetClassType(UInteractiveToolsPresetCollectionAsset::StaticClass())
+						.OnSelectionChanged(this, &FModelingToolsEditorModeToolkit::HandlePresetAssetChanged)
+						.ToolTipText(LOCTEXT("ToolPresets_CreatePresetCollection_Tooltip", "The asset in which to store this new preset."))
+						.FlyoutTileSize(FVector2D(80, 80))
+						.ComboButtonTileSize(FVector2D(80, 80))
+						.InitiallySelectedAsset( FAssetData(FName("/ToolPresets/Presets/_DefaultCollection"),
+															FName("/ToolPresets/Presets"),
+															FName("_DefaultCollection"),
+															UInteractiveToolsPresetCollectionAsset::StaticClass()->GetClassPathName() ))
+					]
+					+ SHorizontalBox::Slot()
+					[
+						SNew(SVerticalBox)
+						+ SVerticalBox::Slot()
+						.AutoHeight()
+						.Padding(10, 5)
+						[
+							SNew(STextBlock)
+							.Text(LOCTEXT("ToolPresets_CreatePresetCollection", "Collection"))
+							.Font(FSlateFontInfo(FCoreStyle::GetDefaultFont(),12, "Bold"))
+						]
+						+ SVerticalBox::Slot()
+						.AutoHeight()
+						.Padding(10, 5)
+						[
+							SNew(STextBlock)
+							.Text_Lambda([this]() { return CurrentPresetLabel; })
+						]
 
-	return Content->AsShared();
-}
+					]
+				]		
+			],
+			FolderDialogArguments, true);
 
-TSharedRef<SWidget> FModelingToolsEditorModeToolkit::GetPresetCreateButtonContent()
-{
+		return FReply::Handled();
+	};
 
 	constexpr bool bShouldCloseMenuAfterSelection = true;
 	FMenuBuilder MenuBuilder(bShouldCloseMenuAfterSelection, nullptr);
@@ -962,74 +1011,105 @@ TSharedRef<SWidget> FModelingToolsEditorModeToolkit::GetPresetCreateButtonConten
 
 
 	{
-		MenuBuilder.BeginSection(NoExtensionHook, LOCTEXT("ListPresetMenu", "Presets for Tool"));
 		FMenuEntryParams MenuEntryParams;
 
+		typedef TMap<FString, TArray<TSharedPtr<FToolPresetOption>>> FPresetsByNameMap;
+		FPresetsByNameMap PresetsByCollectionName;
 		for (TSharedPtr<FToolPresetOption> ToolPresetOption : AvailablePresetsForTool)
 		{
-			FUIAction ApplyPresetAction;
-			ApplyPresetAction.ExecuteAction = FExecuteAction::CreateLambda([this, ToolPresetOption]()
-				{
-					LoadPresetFromCollection(ToolPresetOption->PresetName, ToolPresetOption->PresetCollection);
-				});
-
-			ApplyPresetAction.CanExecuteAction = FCanExecuteAction::CreateLambda([this]()
-				{
-					return this->OwningEditorMode->GetToolManager()->GetActiveTool(EToolSide::Left) != nullptr;
-				});
-
-			MenuBuilder.AddMenuEntry(
-				FText::FromString(ToolPresetOption->PresetLabel),
-				FText::FromString(ToolPresetOption->PresetTooltip),
-				ToolPresetOption->PresetIcon,
-				ApplyPresetAction);
+			UInteractiveToolsPresetCollectionAsset* Preset = nullptr;
+			if (ToolPresetOption->PresetCollection.IsAsset())
+			{
+				Preset = Cast<UInteractiveToolsPresetCollectionAsset>(ToolPresetOption->PresetCollection.TryLoad());
+			}
+			if (Preset)
+			{
+				PresetsByCollectionName.FindOrAdd(Preset->CollectionLabel.ToString()).Add(ToolPresetOption);
+			}
 		}
 
-		MenuBuilder.EndSection();
-	}
-
-	{
-		MenuBuilder.BeginSection(NoExtensionHook, LOCTEXT("CreatePresetMenu", "Create New Preset"));
-		FMenuEntryParams MenuEntryParams;
-
-		TSharedRef<SWidget> PresetName = SNew(SEditableTextBox)
-			.OnTextCommitted_Lambda([this](const FText& NewLabel, ETextCommit::Type&) { NewPresetLabel = NewLabel.ToString(); });
-
-		MenuBuilder.AddWidget(PresetName, LOCTEXT("NewPresetLabelLabel", "Label"), bNoIndent, bSearchable, LOCTEXT("NewPresetLabelTooltip", "A label that identifies the preset in the preset dropdown."));
-
-		TSharedRef<SObjectPropertyEntryBox> PresetCollection = SNew(SObjectPropertyEntryBox)
-			.AllowedClass(UInteractiveToolsPresetCollectionAsset::StaticClass())
-			.ObjectPath_Lambda([this]() {return FModelingToolsEditorModeToolkit::GetCurrentPresetPath(); })
-			.OnShouldFilterAsset(this, &FModelingToolsEditorModeToolkit::HandleFilterPresetAsset)
-			.OnObjectChanged(this, &FModelingToolsEditorModeToolkit::HandlePresetAssetChanged);
-
-		MenuBuilder.AddWidget(PresetCollection, LOCTEXT("NewPresetCollectionLabel", "Collection"), bNoIndent, bSearchable, LOCTEXT("NewPresetCollectionTooltip", "The preset collection asset that should store this preset."));
-
-		TSharedRef<SEditableTextBox> PresetToolTip = SNew(SEditableTextBox)
-			.OnTextCommitted_Lambda([this](const FText& NewToolTip, ETextCommit::Type&) { NewPresetTooltip = NewToolTip.ToString(); });
+		for (FPresetsByNameMap::TConstIterator Iterator = PresetsByCollectionName.CreateConstIterator(); Iterator; ++Iterator)
+		{
+			MenuBuilder.BeginSection(NoExtensionHook, FText::FromString(Iterator.Key()));
+			for (const TSharedPtr<FToolPresetOption>& ToolPresetOption : Iterator.Value())
+			{
+				FUIAction ApplyPresetAction;
+				ApplyPresetAction.ExecuteAction = FExecuteAction::CreateLambda([this, ToolPresetOption]()
+					{
+						LoadPresetFromCollection(ToolPresetOption->PresetIndex, ToolPresetOption->PresetCollection);
+					});
 
-		MenuBuilder.AddWidget(PresetToolTip, LOCTEXT("NewPresetTooltipLabel", "Tooltip"), bNoIndent, bSearchable, LOCTEXT("NewPresetTiooltipTooltip", "A descriptive tooltip that describes the purpose of this preset."));
+				ApplyPresetAction.CanExecuteAction = FCanExecuteAction::CreateLambda([this, ToolPresetOption]()
+					{
+						return this->OwningEditorMode->GetToolManager()->GetActiveTool(EToolSide::Left) != nullptr;
+					});
 
-		FUIAction AddPresetAction;
-		AddPresetAction.ExecuteAction = FExecuteAction::CreateLambda([this, &PresetName, &MenuBuilder]()
-			{
-				CreateNewPresetInCollection(NewPresetLabel,
-					CurrentPreset,
-					NewPresetTooltip,
-					NewPresetIcon);
-			});
+				MenuBuilder.AddMenuEntry(
+					FText::FromString(ToolPresetOption->PresetLabel),
+					FText::FromString(ToolPresetOption->PresetTooltip),
+					ToolPresetOption->PresetIcon,
+					ApplyPresetAction);
+			}
+			MenuBuilder.EndSection();
+		}
 
-		AddPresetAction.CanExecuteAction = FCanExecuteAction::CreateLambda([this]()
-			{
-				return this->OwningEditorMode->GetToolManager()->GetActiveTool(EToolSide::Left) != nullptr && CurrentPreset.IsValid();
-			});
+		MenuBuilder.BeginSection(NoExtensionHook, LOCTEXT("ModelingPresetPanelHeaderManagePresets", "Manage Presets"));
+		MenuBuilder.AddWidget(
+			SNew(SButton)
+			.OnClicked(FOnClicked::CreateLambda(OpenNewPresetDialog))
+			.Content()
+			[
+				SNew(SHorizontalBox)
+				+ SHorizontalBox::Slot()
+			.AutoWidth()
+			.VAlign(VAlign_Center)
+			.Padding(4.0, 0.0f)
+			[
+				SNew(SImage)
+				.ColorAndOpacity(FSlateColor::UseForeground())
+			.Image(FAppStyle::Get().GetBrush("Icons.Plus"))
+			]
+		+ SHorizontalBox::Slot()
+			.AutoWidth()
+			.VAlign(VAlign_Center)
+			.Padding(4.0, 0.0f)
+			[
+				SNew(STextBlock)
+				.Text(LOCTEXT("ModelingPresetPanelCreateNewPreset", "Create New Preset"))
+				.Justification(ETextJustify::Center)
+				.Font(FAppStyle::Get().GetFontStyle("EditorModesPanel.CategoryFontStyle"))
+			]
+			], FText());
 
-		MenuBuilder.AddMenuEntry(LOCTEXT("CreatePresetEntry", "Create Preset"),
-			LOCTEXT("CreatePresetToolTip", "Add a new preset for this tool from current settings."),
-			FSlateIcon(FAppStyle::Get().GetStyleSetName(), "Icons.Plus"),
-			AddPresetAction);
+		MenuBuilder.AddWidget(
+			SNew(SButton)
+			.OnClicked_Lambda([this]() {IPresetEditorModule::Get().ExecuteOpenPresetEditor(); return FReply::Handled(); })
+			.Content()
+			[
+				SNew(SHorizontalBox)
+				+ SHorizontalBox::Slot()
+				.AutoWidth()
+				.VAlign(VAlign_Center)
+				.Padding(4.0, 0.0f)
+				[
+					SNew(SImage)
+					.ColorAndOpacity(FSlateColor::UseForeground())
+					.Image(FAppStyle::Get().GetBrush("Icons.Settings"))
+				]
+				+ SHorizontalBox::Slot()
+				.AutoWidth()
+				.VAlign(VAlign_Center)
+				.Padding(4.0, 0.0f)
+				[
+					SNew(STextBlock)
+					.Text(LOCTEXT("ModelingPresetPanelOpenPresetManager", "Preset Manager"))
+					.Justification(ETextJustify::Center)
+					.Font(FAppStyle::Get().GetFontStyle("EditorModesPanel.CategoryFontStyle"))
+				]
+			], FText());
 
 		MenuBuilder.EndSection();
+
 	}
 
 	return MenuBuilder.MakeWidget();
@@ -1041,14 +1121,42 @@ void FModelingToolsEditorModeToolkit::ClearPresetComboList()
 }
 
 void FModelingToolsEditorModeToolkit::RebuildPresetListForTool(bool bSettingsOpened)
-{
-	if (bSettingsOpened)
+{	
+	TObjectPtr<UPresetUserSettings> UserSettings = UPresetUserSettings::Get();
+	UserSettings->LoadEditorConfig();
+
+	// We need to generate a combined list of Project Loaded and User available presets to intersect the enabled set against...
+	const UPresetProjectSettings* ProjectSettings = GetDefault<UPresetProjectSettings>();
+	TSet<FSoftObjectPath> AllUserPresets;
+	
+	FAssetRegistryModule& AssetRegistryModule = FModuleManager::LoadModuleChecked<FAssetRegistryModule>("AssetRegistry");
+	
+	TArray<FAssetData> AssetData;
+	FARFilter Filter;
+	Filter.ClassPaths.Add(UInteractiveToolsPresetCollectionAsset::StaticClass()->GetClassPathName());
+	Filter.PackagePaths.Add(FName("/ToolPresets"));
+	Filter.bRecursiveClasses = false;
+	Filter.bRecursivePaths = true;
+	Filter.bIncludeOnlyOnDiskAssets = false;
+	
+	AssetRegistryModule.Get().GetAssets(Filter, AssetData);
+	for (int i = 0; i < AssetData.Num(); i++) {
+		UInteractiveToolsPresetCollectionAsset* Object = Cast<UInteractiveToolsPresetCollectionAsset>(AssetData[i].GetAsset());
+		if (Object)
+		{
+			AllUserPresets.Add(Object->GetPathName());
+		}
+	}
+
+	TSet<FSoftObjectPath> AvailablePresetCollections = UserSettings->EnabledPresetCollections.Intersect( ProjectSettings->LoadedPresetCollections.Union(AllUserPresets));
+	if (UserSettings->bDefaultCollectionEnabled)
 	{
-		return;
+		AvailablePresetCollections.Add(UserSettings->DefaultCollectionPath);
 	}
 
+
 	AvailablePresetsForTool.Empty();
-	for (FSoftObjectPath& PresetCollection : PresetSettings->ActivePresetCollectionsPaths)
+	for (const FSoftObjectPath& PresetCollection : AvailablePresetCollections)
 	{
 		FModelingToolsEditorModeToolkitLocals::ExecuteWithPresetAndTool(*OwningEditorMode, EToolSide::Left, PresetCollection,
 			[this, PresetCollection](UInteractiveToolsPresetCollectionAsset& Preset, UInteractiveTool& Tool) {
@@ -1057,16 +1165,18 @@ void FModelingToolsEditorModeToolkit::RebuildPresetListForTool(bool bSettingsOpe
 				{
 					return;
 				}
-				TArray<FString> PresetsForTool;
-				Preset.PerToolPresets[Tool.GetClass()->GetName()].NamedPresets.GenerateKeyArray(PresetsForTool);
-				AvailablePresetsForTool.Reserve(PresetsForTool.Num());
-				for (FString& PresetEntry : PresetsForTool)
+				AvailablePresetsForTool.Reserve(Preset.PerToolPresets[Tool.GetClass()->GetName()].NamedPresets.Num());
+				for (int32 PresetIndex = 0; PresetIndex < Preset.PerToolPresets[Tool.GetClass()->GetName()].NamedPresets.Num(); ++PresetIndex)
 				{
+					const FInteractiveToolPresetDefintion& PresetDef = Preset.PerToolPresets[Tool.GetClass()->GetName()].NamedPresets[PresetIndex];
+					if (!PresetDef.IsValid())
+					{
+						continue;
+					}
 					TSharedPtr<FToolPresetOption> NewOption = MakeShared<FToolPresetOption>();
-					NewOption->PresetLabel = Preset.PerToolPresets[Tool.GetClass()->GetName()].NamedPresets[PresetEntry].Label;
-					NewOption->PresetTooltip = Preset.PerToolPresets[Tool.GetClass()->GetName()].NamedPresets[PresetEntry].Tooltip;
-					//NewOption->PresetIcon = Preset.PerToolPresets[Tool.GetClass()->GetName()].Store[PresetEntry].Icon;
-					NewOption->PresetName = PresetEntry;
+					NewOption->PresetLabel = PresetDef.Label;
+					NewOption->PresetTooltip = PresetDef.Tooltip;
+					NewOption->PresetIndex = PresetIndex;
 					NewOption->PresetCollection = PresetCollection;
 
 					AvailablePresetsForTool.Add(NewOption);
@@ -1078,6 +1188,17 @@ void FModelingToolsEditorModeToolkit::RebuildPresetListForTool(bool bSettingsOpe
 void FModelingToolsEditorModeToolkit::HandlePresetAssetChanged(const FAssetData& InAssetData)
 {
 	CurrentPreset.SetPath(InAssetData.GetObjectPathString());
+
+	UInteractiveToolsPresetCollectionAsset* Preset = nullptr;
+	if (CurrentPreset.IsAsset())
+	{
+		Preset = Cast<UInteractiveToolsPresetCollectionAsset>(CurrentPreset.TryLoad());
+	}
+	if (Preset)
+	{
+		CurrentPresetLabel = Preset->CollectionLabel;
+	}
+	
 }
 
 bool FModelingToolsEditorModeToolkit::HandleFilterPresetAsset(const FAssetData& InAssetData)
@@ -1085,46 +1206,22 @@ bool FModelingToolsEditorModeToolkit::HandleFilterPresetAsset(const FAssetData&
 	return false;
 }
 
-void FModelingToolsEditorModeToolkit::LoadPresetFromCollection(const FString& PresetName, FSoftObjectPath CollectionPath)
+void FModelingToolsEditorModeToolkit::LoadPresetFromCollection(const int32 PresetIndex, FSoftObjectPath CollectionPath)
 {
 	FModelingToolsEditorModeToolkitLocals::ExecuteWithPresetAndTool(*OwningEditorMode, EToolSide::Left, CollectionPath,
-	[this, PresetName](UInteractiveToolsPresetCollectionAsset& Preset, UInteractiveTool& Tool) {
+	[this, PresetIndex](UInteractiveToolsPresetCollectionAsset& Preset, UInteractiveTool& Tool) {
 		TArray<UObject*> PropertySets = Tool.GetToolProperties();
 		if (!Preset.PerToolPresets.Contains(Tool.GetClass()->GetName()))			
 		{
 			return;
 		}
-		if(!Preset.PerToolPresets[Tool.GetClass()->GetName()].NamedPresets.Contains(PresetName))
+		if(PresetIndex < 0 || PresetIndex >= Preset.PerToolPresets[Tool.GetClass()->GetName()].NamedPresets.Num())
 		{
 			return;
 		}
 
-		TArray<TObjectPtr<UObject>>& PropertyStore = Preset.PerToolPresets[Tool.GetClass()->GetName()].NamedPresets.Find(*PresetName)->Properties;
+		Preset.PerToolPresets[Tool.GetClass()->GetName()].NamedPresets[PresetIndex].LoadStoredPropertyData(PropertySets);
 
-		for (UObject* PropertySet : PropertySets)
-		{
-			UObject* StoredPropertySet = nullptr;
-			for (UObject* PotentialStoredPropertySet : PropertyStore)
-			{
-				if (PotentialStoredPropertySet && PotentialStoredPropertySet->GetClass() == PropertySet->GetClass())
-				{
-					StoredPropertySet = PotentialStoredPropertySet;
-				}
-			}
-			if(StoredPropertySet != nullptr)
-			{
-				for (FProperty* Prop : TFieldRange<FProperty>(PropertySet->GetClass()))
-				{
-					void* DestValue = Prop->ContainerPtrToValuePtr<void>(PropertySet);
-					void* SrcValue = Prop->ContainerPtrToValuePtr<void>(StoredPropertySet);
-					if (!Prop->Identical(DestValue, SrcValue))
-					{
-						Prop->CopySingleValue(DestValue, SrcValue);
-						Tool.OnPropertyModified(PropertySet, Prop);
-					}
-				}
-			}
-		}
 	});
 }
 
@@ -1134,40 +1231,27 @@ void FModelingToolsEditorModeToolkit::CreateNewPresetInCollection(const FString&
 	[this, PresetLabel, ToolTip, Icon, CollectionPath](UInteractiveToolsPresetCollectionAsset& Preset, UInteractiveTool& Tool) {
 
 		TArray<UObject*> PropertySets = Tool.GetToolProperties();
-		FString NewPresetName = FString::Printf(TEXT("NewPreset_%d"), Preset.PerToolPresets.FindOrAdd(Tool.GetClass()->GetName()).NamedPresets.Num());
+		Preset.PerToolPresets.FindOrAdd(Tool.GetClass()->GetName()).ToolLabel = ActiveToolName;
+		if (ensure(ActiveToolIcon))
+		{
+			Preset.PerToolPresets.FindOrAdd(Tool.GetClass()->GetName()).ToolIcon = *ActiveToolIcon;
+		}
 		FInteractiveToolPresetDefintion PresetValuesToCreate;
 		PresetValuesToCreate.Label = PresetLabel;
 		PresetValuesToCreate.Tooltip = ToolTip;
 		//PresetValuesToCreate.Icon = Icon;
 
-		TArray<TObjectPtr<UObject>>& PropertyStore = PresetValuesToCreate.Properties;
-		PropertyStore.Empty();
+		PresetValuesToCreate.SetStoredPropertyData(PropertySets);
 
-		for (UObject* PropertySet : PropertySets)
-		{
-			UObject* StoredPropertySet = NewObject<UObject>(&Preset, PropertySet->GetClass());			
-			if (StoredPropertySet != nullptr)
-			{
-				StoredPropertySet->ClearFlags(EObjectFlags::RF_Transient);
-				for (FProperty* Prop : TFieldRange<FProperty>(PropertySet->GetClass()))
-				{
-					void* SrcValue = Prop->ContainerPtrToValuePtr<void>(PropertySet);
-					void* DestValue = Prop->ContainerPtrToValuePtr<void>(StoredPropertySet);
-					Prop->CopySingleValue(DestValue, SrcValue);
-				}
-				PropertyStore.Add(StoredPropertySet);
-			}
-			
-		}
 
-		Preset.PerToolPresets[Tool.GetClass()->GetName()].NamedPresets.Add(NewPresetName, PresetValuesToCreate);
+		Preset.PerToolPresets[Tool.GetClass()->GetName()].NamedPresets.Add(PresetValuesToCreate);
 		Preset.MarkPackageDirty();
 
 		// Finally add this to the current tool's preset list, since we know it should be there.
 		TSharedPtr<FToolPresetOption> NewOption = MakeShared<FToolPresetOption>();
 		NewOption->PresetLabel = PresetLabel;
 		NewOption->PresetTooltip = ToolTip;
-		NewOption->PresetName = NewPresetName;
+		NewOption->PresetIndex = Preset.PerToolPresets[Tool.GetClass()->GetName()].NamedPresets.Num() - 1;
 		NewOption->PresetCollection = CollectionPath;
 
 		AvailablePresetsForTool.Add(NewOption);
@@ -1175,38 +1259,37 @@ void FModelingToolsEditorModeToolkit::CreateNewPresetInCollection(const FString&
 	});
 }
 
-void FModelingToolsEditorModeToolkit::SaveActivePreset()
+void FModelingToolsEditorModeToolkit::UpdatePresetInCollection(const FToolPresetOption& PresetToEditIn, bool bUpdateStoredPresetValues)
 {
-	
-	// Saving presets has been disabled for the moment until the 3rd UI generation is done, but we'll want this logic then.
-#if 0
-	FModelingToolsEditorModeToolkitLocals::ExecuteWithPresetAndTool(*OwningEditorMode, EToolSide::Left, CurrentPreset,
-	[this](UInteractiveToolsPresetCollectionAsset& Preset, UInteractiveTool& Tool) {
-		TArray<UObject*> PropertySets = Tool.GetToolProperties();
-		TArray<TObjectPtr<UObject>>& PropertyStore = Preset.PerToolPresets.FindOrAdd(Tool.GetClass()->GetName()).NamedPresets.FindOrAdd(*ActiveNamedPreset).Properties;
-		PropertyStore.Empty();
+	FModelingToolsEditorModeToolkitLocals::ExecuteWithPresetAndTool(*OwningEditorMode, EToolSide::Left, PresetToEditIn.PresetCollection,
+		[this, PresetToEditIn, bUpdateStoredPresetValues](UInteractiveToolsPresetCollectionAsset& Preset, UInteractiveTool& Tool) {
+			TArray<UObject*> PropertySets = Tool.GetToolProperties();
+			if (!Preset.PerToolPresets.Contains(Tool.GetClass()->GetName()))
+			{
+				return;
+			}
+			if (PresetToEditIn.PresetIndex < 0 || PresetToEditIn.PresetIndex >= Preset.PerToolPresets[Tool.GetClass()->GetName()].NamedPresets.Num() )
+			{
+				return;
+			}
 
-		for (UObject* PropertySet : PropertySets)
-		{
-			UObject* StoredPropertySet = NewObject<UObject>(&Preset, PropertySet->GetClass());
-			StoredPropertySet->ClearFlags(EObjectFlags::RF_Transient);
-			if (StoredPropertySet != nullptr)
+			if (bUpdateStoredPresetValues)
 			{
-				for (FProperty* Prop : TFieldRange<FProperty>(PropertySet->GetClass()))
-				{
-					void* SrcValue = Prop->ContainerPtrToValuePtr<void>(PropertySet);
-					void* DestValue = Prop->ContainerPtrToValuePtr<void>(StoredPropertySet);
-					Prop->CopySingleValue(DestValue, SrcValue);
-				}
+				Preset.PerToolPresets.FindOrAdd(Tool.GetClass()->GetName()).NamedPresets[PresetToEditIn.PresetIndex].SetStoredPropertyData(PropertySets);
+				Preset.MarkPackageDirty();
 			}
-			PropertyStore.Add(StoredPropertySet);
-		}	
 
-		Preset.MarkPackageDirty();
-	});
-#endif
+			Preset.PerToolPresets.FindOrAdd(Tool.GetClass()->GetName()).NamedPresets[PresetToEditIn.PresetIndex].Label = PresetToEditIn.PresetLabel;
+			Preset.PerToolPresets.FindOrAdd(Tool.GetClass()->GetName()).NamedPresets[PresetToEditIn.PresetIndex].Tooltip = PresetToEditIn.PresetTooltip;
+
+		});
+
+	RebuildPresetListForTool(false);
 }
 
+
+
+
 void FModelingToolsEditorModeToolkit::InitializeAfterModeSetup()
 {
 	if (bFirstInitializeAfterModeSetup)
diff --git a/Engine/Plugins/Editor/ModelingToolsEditorMode/Source/ModelingToolsEditorMode/Public/ModelingToolsEditorModeToolkit.h b/Engine/Plugins/Editor/ModelingToolsEditorMode/Source/ModelingToolsEditorMode/Public/ModelingToolsEditorModeToolkit.h
--- a/Engine/Plugins/Editor/ModelingToolsEditorMode/Source/ModelingToolsEditorMode/Public/ModelingToolsEditorModeToolkit.h
+++ b/Engine/Plugins/Editor/ModelingToolsEditorMode/Source/ModelingToolsEditorMode/Public/ModelingToolsEditorModeToolkit.h
@@ -4,9 +4,6 @@
 
 #include "Toolkits/BaseToolkit.h"
 #include "StatusBarSubsystem.h"
-#include "UObject/GCObject.h"
-
-#include "ModelingToolsEditorModeToolkit.generated.h"
 
 class IAssetViewport;
 class SEditableTextBox;
@@ -22,26 +19,16 @@ class STextBlock;
 class UGeometrySelectionManager;
 class UInteractiveToolsPresetCollectionAsset;
 
-UCLASS()
-class UPresetSettingsProperties : public UObject
-{
-	GENERATED_BODY()
-public:
-
-	UPROPERTY(EditAnywhere, Category = "Presets", meta=(MetaClass = "/Script/PresetAsset.InteractiveToolsPresetCollectionAsset"))
-	TArray< FSoftObjectPath > ActivePresetCollectionsPaths;
-};
-
 struct FToolPresetOption
 {
 	FString PresetLabel;
 	FString PresetTooltip;
 	FSlateIcon PresetIcon;
-	FString PresetName;
+	int32 PresetIndex;
 	FSoftObjectPath PresetCollection;
 };
 
-class FModelingToolsEditorModeToolkit : public FModeToolkit, public FGCObject
+class FModelingToolsEditorModeToolkit : public FModeToolkit
 {
 public:
 
@@ -107,10 +94,6 @@ public:
 	// in a non-member function in ModelingModeToolkit_Toolbars.cpp
 	TSharedPtr<STransformGizmoNumericalUIOverlay> GetGizmoNumericalUIOverlayWidget() { return GizmoNumericalUIOverlayWidget; }
 
-	/** GCObject interface */
-	virtual void AddReferencedObjects(FReferenceCollector& Collector) override;
-	virtual FString GetReferencerName() const override;
-
 private:
 	const static TArray<FName> PaletteNames_Standard;
 
@@ -165,9 +148,11 @@ private:
 	void SelectNewAssetPath() const;
 
 	// Presets
-	TObjectPtr<UPresetSettingsProperties> PresetSettings;
+	TSharedPtr<SWidget> ToolPresetArea;
 	TSharedPtr<SWidget> MakePresetPanel();
 	FSoftObjectPath CurrentPreset;
+	FText CurrentPresetLabel;
+	FToolPresetOption SelectedPresetToEdit;
 	FString NewPresetLabel;
 	FString NewPresetTooltip;
 	FSlateIcon NewPresetIcon;
@@ -175,16 +160,15 @@ private:
 	TArray<TSharedPtr<FToolPresetOption>> AvailablePresetsForTool;
 	TSharedPtr<SEditableComboBox<TSharedPtr<FString>>> PresetComboBox;
 
-	TSharedRef<SWidget> GetPresetSettingsButtonContent();
 	TSharedRef<SWidget> GetPresetCreateButtonContent();
 
 	void CreateNewPresetInCollection(const FString& PresetLabel, FSoftObjectPath CollectionPath, const FString& ToolTip, FSlateIcon Icon);
-	void LoadPresetFromCollection(const FString& PresetName, FSoftObjectPath CollectionPath);
+	void LoadPresetFromCollection(const int32 PresetIndex, FSoftObjectPath CollectionPath);
+	void UpdatePresetInCollection(const FToolPresetOption& PresetToEditIn, bool bUpdateStoredPresetValues);
 
 	FString GetCurrentPresetPath() { return CurrentPreset.GetAssetPathString(); }
 	void HandlePresetAssetChanged(const FAssetData& InAssetData);
 	bool HandleFilterPresetAsset(const FAssetData& InAssetData);
-	void SaveActivePreset();
 	
 	void RebuildPresetListForTool(bool bSettingsOpened);
 	TSharedRef<SWidget> MakePresetComboWidget(TSharedPtr<FString> InItem);
diff --git a/Engine/Plugins/Experimental/MeshModelingToolsetExp/Source/ModelingEditorUI/Private/ModelingWidgets/SToolInputAssetComboPanel.cpp b/Engine/Plugins/Experimental/MeshModelingToolsetExp/Source/ModelingEditorUI/Private/ModelingWidgets/SToolInputAssetComboPanel.cpp
--- a/Engine/Plugins/Experimental/MeshModelingToolsetExp/Source/ModelingEditorUI/Private/ModelingWidgets/SToolInputAssetComboPanel.cpp
+++ b/Engine/Plugins/Experimental/MeshModelingToolsetExp/Source/ModelingEditorUI/Private/ModelingWidgets/SToolInputAssetComboPanel.cpp
@@ -30,6 +30,7 @@ void SToolInputAssetComboPanel::Construct(const FArguments& InArgs)
 	this->Property = InArgs._Property;
 	this->CollectionSets = InArgs._CollectionSets;
 	this->OnSelectionChanged = InArgs._OnSelectionChanged;
+	RecentAssetsProvider = InArgs._RecentAssetsProvider;
 
 	// validate arguments
 	if ( InArgs._AssetClassType == nullptr )
@@ -57,9 +58,6 @@ void SToolInputAssetComboPanel::Construct(const FArguments& InArgs)
 			AssetThumbnail->MakeThumbnailWidget(ThumbnailConfig)
 		];
 
-	RecentAssetsProvider = InArgs._RecentAssetsProvider;
-
-
 	ComboButton = SNew(SComboButton)
 		.ToolTipText(UseTooltipText)
 		.HasDownArrow(false)
@@ -223,6 +221,10 @@ void SToolInputAssetComboPanel::Construct(const FArguments& InArgs)
 		}
 	}
 
+	if (InArgs._InitiallySelectedAsset.IsValid())
+	{
+		NewAssetSelected(InArgs._InitiallySelectedAsset);
+	}
 
 	ChildSlot
 		[
diff --git a/Engine/Plugins/Experimental/MeshModelingToolsetExp/Source/ModelingEditorUI/Public/ModelingWidgets/SToolInputAssetComboPanel.h b/Engine/Plugins/Experimental/MeshModelingToolsetExp/Source/ModelingEditorUI/Public/ModelingWidgets/SToolInputAssetComboPanel.h
--- a/Engine/Plugins/Experimental/MeshModelingToolsetExp/Source/ModelingEditorUI/Public/ModelingWidgets/SToolInputAssetComboPanel.h
+++ b/Engine/Plugins/Experimental/MeshModelingToolsetExp/Source/ModelingEditorUI/Public/ModelingWidgets/SToolInputAssetComboPanel.h
@@ -65,6 +65,7 @@ public:
 		, _FlyoutSize(600, 400)
 		, _AssetClassType(0)
 		, _OnSelectionChanged()
+		, _InitiallySelectedAsset()
 	{
 		}
 
@@ -95,6 +96,9 @@ public:
 		/** This delegate is executed each time the Selected Asset is modified */
 		SLATE_EVENT( FOnSelectedAssetChanged, OnSelectionChanged )
 
+		/** Sets the asset selected by the widget before any user made selection occurs. */
+		SLATE_ARGUMENT( FAssetData, InitiallySelectedAsset)
+
 	SLATE_END_ARGS()
 
 
diff --git a/Engine/Plugins/Experimental/ToolPresets/Source/PresetAsset/PresetAsset.Build.cs b/Engine/Plugins/Experimental/ToolPresets/Source/PresetAsset/PresetAsset.Build.cs
--- a/Engine/Plugins/Experimental/ToolPresets/Source/PresetAsset/PresetAsset.Build.cs
+++ b/Engine/Plugins/Experimental/ToolPresets/Source/PresetAsset/PresetAsset.Build.cs
@@ -39,6 +39,8 @@ public class PresetAsset: ModuleRules
 				"SlateCore",
 				"UnrealEd",
 				"AssetDefinition",
+				"JsonUtilities",
+				"Json",
 			}
 			);
 		
diff --git a/Engine/Plugins/Experimental/ToolPresets/Source/PresetAsset/Private/PresetAsset.cpp b/Engine/Plugins/Experimental/ToolPresets/Source/PresetAsset/Private/PresetAsset.cpp
new file mode 100644
--- /dev/null
+++ b/Engine/Plugins/Experimental/ToolPresets/Source/PresetAsset/Private/PresetAsset.cpp
@@ -0,0 +1,64 @@
+// Copyright Epic Games, Inc. All Rights Reserved.
+
+#include "PresetAsset.h"
+#include "JsonObjectConverter.h"
+
+void FInteractiveToolPresetDefintion::SetStoredPropertyData(TArray<UObject*>& Properties)
+{
+	FJsonObjectWrapper JsonWrapper;
+
+	for (UObject* PropertySet : Properties)
+	{
+		TSharedPtr<FJsonObject> PropertyJsonObject = MakeShared<FJsonObject>();
+
+		if (PropertyJsonObject)
+		{
+			for (FProperty* Prop : TFieldRange<FProperty>(PropertySet->GetClass()))
+			{
+				TSharedPtr<FJsonValue> JsonProp = FJsonObjectConverter::UPropertyToJsonValue(Prop, Prop->ContainerPtrToValuePtr<void>(PropertySet));				
+				FString FieldName;
+				Prop->GetName(FieldName);
+				PropertyJsonObject->SetField(FieldName, JsonProp);
+			}			
+		}
+
+		FString ClassName;
+		PropertySet->GetClass()->GetName(ClassName);
+		JsonWrapper.JsonObject->SetObjectField(ClassName, PropertyJsonObject);
+	}
+	
+	JsonWrapper.JsonObjectToString(StoredProperties);
+}
+
+void FInteractiveToolPresetDefintion::LoadStoredPropertyData(TArray<UObject*>& Properties)
+{
+	FJsonObjectWrapper JsonWrapper;
+	JsonWrapper.JsonObjectFromString(StoredProperties);
+	
+	for (UObject* PropertySet : Properties)
+	{
+		FString ClassName;
+		PropertySet->GetClass()->GetName(ClassName);
+		TSharedPtr<FJsonObject> PropertyJsonObject = JsonWrapper.JsonObject->GetObjectField(ClassName);
+
+		if (PropertyJsonObject)
+		{
+			for (FProperty* Prop : TFieldRange<FProperty>(PropertySet->GetClass()))
+			{
+				FString FieldName;
+				Prop->GetName(FieldName);
+				TSharedPtr<FJsonValue>* JsonField = PropertyJsonObject->Values.Find(FieldName);
+				if (JsonField)
+				{
+					FJsonObjectConverter::JsonValueToUProperty(*JsonField, Prop, Prop->ContainerPtrToValuePtr<void>(PropertySet));
+				}
+			}
+		}
+	}
+}
+
+bool FInteractiveToolPresetDefintion::IsValid() const
+{
+	FJsonObjectWrapper JsonWrapper;
+	return JsonWrapper.JsonObjectFromString(StoredProperties);
+}
\ No newline at end of file
diff --git a/Engine/Plugins/Experimental/ToolPresets/Source/PresetAsset/Public/PresetAsset.h b/Engine/Plugins/Experimental/ToolPresets/Source/PresetAsset/Public/PresetAsset.h
--- a/Engine/Plugins/Experimental/ToolPresets/Source/PresetAsset/Public/PresetAsset.h
+++ b/Engine/Plugins/Experimental/ToolPresets/Source/PresetAsset/Public/PresetAsset.h
@@ -13,12 +13,12 @@
 #define LOCTEXT_NAMESPACE "PresetAsset"
 
 USTRUCT(BlueprintType)
-struct FInteractiveToolPresetDefintion
+struct PRESETASSET_API FInteractiveToolPresetDefintion
 {
 	GENERATED_BODY()
 public:
 	UPROPERTY(VisibleAnywhere, Category = "PresetAsset")
-	TArray< TObjectPtr<UObject> > Properties;
+	FString StoredProperties;
 
 	UPROPERTY(EditAnywhere, Category = "PresetAsset")
 	FString Label;
@@ -26,18 +26,24 @@ public:
 	UPROPERTY(EditAnywhere, Category = "PresetAsset")
 	FString Tooltip;
 
-	// TODO: Investigate how to store things like icons. This doesn't work as written, as FSlateBrush can't be serialized?
-	//UPROPERTY(EditAnywhere, Category = "PresetAsset")
-	//FSlateBrush Icon;
+	bool IsValid() const;
+	void SetStoredPropertyData(TArray<UObject*>& Properties);
+	void LoadStoredPropertyData(TArray<UObject*>& Properties);
 };
 
 USTRUCT(BlueprintType)
-struct FInteractiveToolPresetStore
+struct PRESETASSET_API FInteractiveToolPresetStore
 {
 	GENERATED_BODY()
 public:
 	UPROPERTY(EditAnywhere, Category = "PresetAsset")
-	TMap< FString, FInteractiveToolPresetDefintion> NamedPresets;
+	TArray<FInteractiveToolPresetDefintion> NamedPresets;
+
+	UPROPERTY(VisibleAnywhere, Category = "PresetAsset")
+	FText ToolLabel;
+
+	UPROPERTY(VisibleAnywhere, Category = "PresetAsset")
+	FSlateBrush ToolIcon;
 };
 
 /**
@@ -61,6 +67,8 @@ public:
 	UPROPERTY(VisibleAnywhere, Category = "PresetAsset")
 	TMap<FString, FInteractiveToolPresetStore > PerToolPresets;
 
+	UPROPERTY(EditAnywhere, Category = "PresetAsset")
+	FText CollectionLabel;
 };
 
 
diff --git a/Engine/Plugins/Experimental/ToolPresets/Source/PresetEditor/PresetEditor.Build.cs b/Engine/Plugins/Experimental/ToolPresets/Source/PresetEditor/PresetEditor.Build.cs
new file mode 100644
--- /dev/null
+++ b/Engine/Plugins/Experimental/ToolPresets/Source/PresetEditor/PresetEditor.Build.cs
@@ -0,0 +1,37 @@
+// Copyright Epic Games, Inc. All Rights Reserved.
+
+namespace UnrealBuildTool.Rules
+{
+	public class PresetEditor : ModuleRules
+	{
+		public PresetEditor(ReadOnlyTargetRules Target) : base(Target)
+		{
+			PublicDependencyModuleNames.AddRange(
+				new string[] {
+					"Core",
+					"EditorConfig",
+				}
+			);
+
+			PrivateDependencyModuleNames.AddRange(
+				new string[] {
+					"ApplicationCore",
+					"CoreUObject",
+					"EditorFramework",
+					"EditorStyle",
+					"Engine",
+					"InputCore",
+					"Slate",
+					"SlateCore",
+					"ToolWidgets",
+					"PresetAsset",
+					"UnrealEd",
+					"Projects",
+					"DeveloperSettings",
+					"ContentBrowser",
+					"ContentBrowserData",
+				}
+			);
+		}
+	}
+}
diff --git a/Engine/Plugins/Experimental/ToolPresets/Source/PresetEditor/Private/PresetEditorModule.cpp b/Engine/Plugins/Experimental/ToolPresets/Source/PresetEditor/Private/PresetEditorModule.cpp
new file mode 100644
--- /dev/null
+++ b/Engine/Plugins/Experimental/ToolPresets/Source/PresetEditor/Private/PresetEditorModule.cpp
@@ -0,0 +1,133 @@
+// Copyright Epic Games, Inc. All Rights Reserved.
+
+#include "Delegates/Delegate.h"
+#include "Framework/Docking/TabManager.h"
+#include "IPresetEditorModule.h"
+#include "Internationalization/Internationalization.h"
+#include "Modules/ModuleManager.h"
+#include "Styling/AppStyle.h"
+#include "Templates/SharedPointer.h"
+#include "Textures/SlateIcon.h"
+#include "UObject/NameTypes.h"
+#include "Widgets/DeclarativeSyntaxSupport.h"
+#include "Widgets/Docking/SDockTab.h"
+#include "Widgets/SPresetManager.h"
+
+#include "ContentBrowserModule.h"
+#include "IContentBrowserSingleton.h"
+#include "UObject/SavePackage.h"
+#include "AssetRegistry/AssetRegistryModule.h"
+#include "PresetAsset.h"
+#include "UObject/UObjectGlobals.h"
+#include "PackageTools.h"
+#include "PresetEditorStyle.h"
+
+#define LOCTEXT_NAMESPACE "PresetEditorModule"
+
+const FName PresetEditorTabName("Preset");
+
+class FPresetEditorModule : public IPresetEditorModule
+{
+public:
+
+	virtual void StartupModule() override
+	{
+		FModuleManager::Get().LoadModuleChecked(TEXT("DataValidation"));
+		FModuleManager::Get().LoadModuleChecked(TEXT("AssetTools"));
+		FModuleManager::Get().LoadModuleChecked(TEXT("ContentBrowser"));
+		FModuleManager::Get().LoadModuleChecked(TEXT("AssetRegistry"));
+
+		FGlobalTabmanager::Get()->RegisterNomadTabSpawner(PresetEditorTabName, FOnSpawnTab::CreateRaw(this, &FPresetEditorModule::HandleSpawnPresetEditorTab))
+			.SetDisplayName(NSLOCTEXT("FPresetModule", "PresetTabTitle", "Preset Manager"))
+			.SetTooltipText(NSLOCTEXT("FPresetModule", "PresetTooltipText", "Open the Preset Manager tab."))
+			.SetIcon(FSlateIcon(FAppStyle::GetAppStyleSetName(), "Preset.TabIcon"))
+			.SetAutoGenerateMenuEntry(false);
+
+		FCoreDelegates::OnPostEngineInit.AddRaw(this, &FPresetEditorModule::OnPostEngineInit);
+
+		
+	}
+	
+	virtual void ShutdownModule() override
+	{
+		FGlobalTabmanager::Get()->UnregisterNomadTabSpawner(PresetEditorTabName);
+
+		FPresetEditorStyle::Shutdown();
+	}
+
+	void OnPostEngineInit()
+	{
+		// Register slate style overrides
+		FPresetEditorStyle::Initialize();
+
+		GenerateDefaultCollection();
+	}
+
+	virtual bool SupportsDynamicReloading() override
+	{
+		return true;
+	}
+
+	virtual void ExecuteOpenPresetEditor() override
+	{
+		FGlobalTabmanager::Get()->TryInvokeTab(PresetEditorTabName);
+	}	
+
+private:
+	
+	/** Handles creating the project settings tab. */
+	TSharedRef<SDockTab> HandleSpawnPresetEditorTab(const FSpawnTabArgs& SpawnTabArgs)
+	{
+		const TSharedRef<SDockTab> DockTab = SNew(SDockTab)
+			.TabRole(ETabRole::NomadTab)
+		[
+			SNew(SPresetManager)
+		];
+		return DockTab;
+	}
+
+	void GenerateDefaultCollection()
+	{
+		// Load necessary modules
+		FAssetToolsModule& AssetToolsModule = FModuleManager::Get().LoadModuleChecked<FAssetToolsModule>("AssetTools");
+		FContentBrowserModule& ContentBrowserModule = FModuleManager::LoadModuleChecked<FContentBrowserModule>("ContentBrowser");
+		FAssetRegistryModule& AssetRegistryModule = FModuleManager::LoadModuleChecked<FAssetRegistryModule>(TEXT("AssetRegistry"));
+		IAssetRegistry& AssetRegistry = AssetRegistryModule.Get();
+
+		TArray<FAssetData> AssetData;
+
+		// Find (or create!) the desired package for this object
+		FString PackageName = "/ToolPresets/Presets/_DefaultCollection";
+		FString ObjectName = "_DefaultCollection";
+		const FString PackagePath = FPackageName::GetLongPackagePath(PackageName);
+
+		// Handle fully loading packages before creating new objects.
+		UPackage* Pkg = CreatePackage(*PackageName);
+		TArray<UPackage*> TopLevelPackages;
+		TopLevelPackages.Add(Pkg);
+		UPackageTools::HandleFullyLoadingPackages(TopLevelPackages, LOCTEXT("CreateANewObject", "Create a new object"));
+
+		// Check for an existing object
+		UObject* ExistingObject = StaticFindObject(UObject::StaticClass(), Pkg, *ObjectName);
+
+		if (!ExistingObject)
+		{
+			// Create object and package
+			UInteractiveToolsPresetCollectionAssetFactory* MyFactory = NewObject<UInteractiveToolsPresetCollectionAssetFactory>(UInteractiveToolsPresetCollectionAssetFactory::StaticClass()); // Can omit, and a default factory will be used
+			UObject* NewObject = AssetToolsModule.Get().CreateAsset(ObjectName, PackagePath, UInteractiveToolsPresetCollectionAsset::StaticClass(), MyFactory);
+			UInteractiveToolsPresetCollectionAsset* NewCollection = ExactCast<UInteractiveToolsPresetCollectionAsset>(NewObject);
+			NewCollection->CollectionLabel = FText::FromString("Default Collection");
+			FSavePackageArgs SavePackageArgs;
+			SavePackageArgs.TopLevelFlags = RF_Public | RF_Standalone;
+			UPackage::Save(Pkg, NewObject, *FPackageName::LongPackageNameToFilename(PackageName, FPackageName::GetAssetPackageExtension()), SavePackageArgs);
+
+			// Inform asset registry
+			AssetRegistry.AssetCreated(NewObject);
+		}
+	}
+};
+
+#undef LOCTEXT_NAMESPACE
+
+IMPLEMENT_MODULE(FPresetEditorModule, PresetEditor);
+
diff --git a/Engine/Plugins/Experimental/ToolPresets/Source/PresetEditor/Private/PresetEditorStyle.cpp b/Engine/Plugins/Experimental/ToolPresets/Source/PresetEditor/Private/PresetEditorStyle.cpp
new file mode 100644
--- /dev/null
+++ b/Engine/Plugins/Experimental/ToolPresets/Source/PresetEditor/Private/PresetEditorStyle.cpp
@@ -0,0 +1,92 @@
+// Copyright Epic Games, Inc. All Rights Reserved.
+
+#include "PresetEditorStyle.h"
+#include "Brushes/SlateImageBrush.h"
+#include "Styling/SlateStyleRegistry.h"
+#include "Misc/Paths.h"
+#include "Styling/CoreStyle.h"
+#include "Interfaces/IPluginManager.h"
+#include "SlateOptMacros.h"
+#include "Styling/SlateStyle.h"
+#include "Styling/SlateStyleMacros.h"
+#include "Styling/ToolBarStyle.h"
+
+
+#define IMAGE_PLUGIN_BRUSH( RelativePath, ... ) FSlateImageBrush( FPresetEditorStyle::InContent( RelativePath, ".png" ), __VA_ARGS__ )
+
+// This is to fix the issue that SlateStyleMacros like IMAGE_BRUSH look for RootToContentDir but StyleSet->RootToContentDir is how this style is set up
+#define RootToContentDir StyleSet->RootToContentDir
+
+FString FPresetEditorStyle::InContent(const FString& RelativePath, const ANSICHAR* Extension)
+{
+	static FString ContentDir = IPluginManager::Get().FindPlugin(TEXT("ToolPresets"))->GetContentDir();
+	return (ContentDir / RelativePath) + Extension;
+}
+
+TSharedPtr< FSlateStyleSet > FPresetEditorStyle::StyleSet = nullptr;
+TSharedPtr< class ISlateStyle > FPresetEditorStyle::Get() { return StyleSet; }
+
+FName FPresetEditorStyle::GetStyleSetName()
+{
+	static FName PresetEditorStyleName(TEXT("PresetEditorStyle"));
+	return PresetEditorStyleName;
+}
+
+const FSlateBrush* FPresetEditorStyle::GetBrush(FName PropertyName, const ANSICHAR* Specifier)
+{
+	return Get()->GetBrush(PropertyName, Specifier);
+}
+
+
+BEGIN_SLATE_FUNCTION_BUILD_OPTIMIZATION
+
+void FPresetEditorStyle::Initialize()
+{
+	// Const icon sizes
+	const FVector2D Icon8x8(8.0f, 8.0f);
+	const FVector2D Icon16x16(16.0f, 16.0f);
+	const FVector2D Icon20x20(20.0f, 20.0f);
+	const FVector2D Icon28x28(28.0f, 28.0f);
+	const FVector2D Icon40x40(40.0f, 40.0f);
+	const FVector2D Icon120(120.0f, 120.0f);
+
+	// Only register once
+	if (StyleSet.IsValid())
+	{
+		return;
+	}
+
+	StyleSet = MakeShareable(new FSlateStyleSet(GetStyleSetName()));
+
+	// If we get asked for something that we don't set, we should default to editor style
+	StyleSet->SetParentStyleName("EditorStyle");
+
+	StyleSet->SetContentRoot(FPaths::EnginePluginsDir() / TEXT("Experimental/ToolPresets/Content"));
+	StyleSet->SetCoreContentRoot(FPaths::EngineContentDir() / TEXT("Slate"));
+
+	// Icons
+	{
+		StyleSet->Set("ManagerIcons.Tools", new IMAGE_PLUGIN_BRUSH("Icons/Primitive_40x", Icon40x40));
+
+
+    }
+
+	FSlateStyleRegistry::RegisterSlateStyle(*StyleSet.Get());
+};
+
+END_SLATE_FUNCTION_BUILD_OPTIMIZATION
+
+#undef IMAGE_PLUGIN_BRUSH
+#undef IMAGE_BRUSH
+#undef BOX_BRUSH
+#undef DEFAULT_FONT
+
+void FPresetEditorStyle::Shutdown()
+{
+	if (StyleSet.IsValid())
+	{
+		FSlateStyleRegistry::UnRegisterSlateStyle(*StyleSet.Get());
+		ensure(StyleSet.IsUnique());
+		StyleSet.Reset();
+	}
+}
\ No newline at end of file
diff --git a/Engine/Plugins/Experimental/ToolPresets/Source/PresetEditor/Private/PresetEditorStyle.h b/Engine/Plugins/Experimental/ToolPresets/Source/PresetEditor/Private/PresetEditorStyle.h
new file mode 100644
--- /dev/null
+++ b/Engine/Plugins/Experimental/ToolPresets/Source/PresetEditor/Private/PresetEditorStyle.h
@@ -0,0 +1,28 @@
+// Copyright Epic Games, Inc. All Rights Reserved.
+
+#pragma once
+
+#include "Templates/SharedPointer.h"
+
+struct FSlateBrush;
+
+class FPresetEditorStyle
+{
+public:
+	static void Initialize();
+
+	static void Shutdown();
+
+	static TSharedPtr< class ISlateStyle > Get();
+
+	static FName GetStyleSetName();
+
+	// use to access icons defined by the style set by name, eg GetBrush("BrushFalloffIcons.Smooth")
+	static const FSlateBrush* GetBrush(FName PropertyName, const ANSICHAR* Specifier = NULL);
+
+private:
+	static FString InContent(const FString& RelativePath, const ANSICHAR* Extension);
+
+private:
+	static TSharedPtr< class FSlateStyleSet > StyleSet;
+};
\ No newline at end of file
diff --git a/Engine/Plugins/Experimental/ToolPresets/Source/PresetEditor/Private/PresetSettings.cpp b/Engine/Plugins/Experimental/ToolPresets/Source/PresetEditor/Private/PresetSettings.cpp
new file mode 100644
--- /dev/null
+++ b/Engine/Plugins/Experimental/ToolPresets/Source/PresetEditor/Private/PresetSettings.cpp
@@ -0,0 +1,21 @@
+/// Copyright Epic Games, Inc. All Rights Reserved.
+
+#include "PresetSettings.h"
+
+#include "UObject/ObjectPtr.h"
+
+TObjectPtr<UPresetUserSettings> UPresetUserSettings::Instance = nullptr;
+
+void UPresetUserSettings::Initialize()
+{
+	if (Instance == nullptr)
+	{
+		Instance = NewObject<UPresetUserSettings>();
+		Instance->AddToRoot();
+	}
+}
+
+UPresetUserSettings* UPresetUserSettings::Get()
+{
+	return Instance;
+}
diff --git a/Engine/Plugins/Experimental/ToolPresets/Source/PresetEditor/Private/Widgets/SPresetManager.cpp b/Engine/Plugins/Experimental/ToolPresets/Source/PresetEditor/Private/Widgets/SPresetManager.cpp
new file mode 100644
--- /dev/null
+++ b/Engine/Plugins/Experimental/ToolPresets/Source/PresetEditor/Private/Widgets/SPresetManager.cpp
@@ -0,0 +1,1165 @@
+// Copyright Epic Games, Inc. All Rights Reserved.
+
+#include "Widgets/SPresetManager.h"
+#include "SlateOptMacros.h"
+#include "Widgets/Images/SImage.h"
+#include "Widgets/Text/STextBlock.h"
+#include "Widgets/Input/SButton.h"
+#include "Styling/AppStyle.h"
+#include "Editor.h"
+#include "Framework/MultiBox/MultiBoxBuilder.h"
+#include "Widgets/Layout/SExpandableArea.h"
+#include "Widgets/Layout/SSpacer.h"
+#include "Widgets/Views/SListView.h"
+#include "Widgets/Views/STreeView.h"
+#include "Widgets/Views/STableRow.h"
+#include "Widgets/Input/SComboButton.h"
+#include "Widgets/Input/SCheckBox.h"
+#include "Widgets/Input/SEditableTextBox.h"
+#include "Widgets/Input/SMultiLineEditableTextBox.h"
+#include "Editor/TransBuffer.h"
+#include "PresetSettings.h"
+#include "Misc/ScopedSlowTask.h"
+#include "AssetRegistry/AssetRegistryModule.h"
+#include "Templates/SharedPointer.h"
+#include "Widgets/SWidget.h"
+#include "ContentBrowserModule.h"
+#include "IContentBrowserSingleton.h"
+#include "UObject/SavePackage.h"
+#include "ObjectTools.h"
+#include "PresetAsset.h"
+#include "SNegativeActionButton.h"
+#include "SPositiveActionButton.h"
+#include "PresetEditorStyle.h"
+
+#define LOCTEXT_NAMESPACE "SPresetManager"
+
+DECLARE_DELEGATE_TwoParams(FOnCollectionEnabledCheckboxChanged, TSharedPtr<SPresetManager::FPresetViewEntry>, ECheckBoxState)
+DECLARE_DELEGATE_TwoParams(FOnPresetLabelChanged, TSharedPtr<SPresetManager::FPresetViewEntry>, FText)
+DECLARE_DELEGATE_TwoParams(FOnPresetTooltipChanged, TSharedPtr<SPresetManager::FPresetViewEntry>, FText)
+DECLARE_DELEGATE_OneParam(FOnPresetDeleted, TSharedPtr<SPresetManager::FPresetViewEntry>)
+
+namespace PresetManagerLocals
+{
+	template<typename ASSETCLASS>
+	void GetObjectsOfClass(TArray<FSoftObjectPath>& OutArray)
+	{
+		FAssetRegistryModule& AssetRegistryModule = FModuleManager::LoadModuleChecked<FAssetRegistryModule>("AssetRegistry");
+		TArray<FAssetData> AssetData;
+
+		FARFilter Filter;
+		Filter.ClassPaths.Add(ASSETCLASS::StaticClass()->GetClassPathName());
+		Filter.PackagePaths.Add(FName("/ToolPresets"));
+		Filter.bRecursiveClasses = false;	
+		Filter.bRecursivePaths = true;
+		Filter.bIncludeOnlyOnDiskAssets = false;
+		
+		AssetRegistryModule.Get().GetAssets(Filter, AssetData);
+
+		for (int i = 0; i < AssetData.Num(); i++) {
+			ASSETCLASS* Object = Cast<ASSETCLASS>(AssetData[i].GetAsset());		
+			if (Object)
+			{				
+				OutArray.Add(Object->GetPathName());
+			}
+		}
+	}
+
+	template<typename ItemType> class SCollectionTableRow;
+
+	template<typename ItemType>
+	class SCollectionTableRow : public STableRow< ItemType>
+	{
+		typedef SCollectionTableRow< ItemType > FSuperRowType;
+		typedef typename STableRow<ItemType>::FArguments FTableRowArgs;
+		typedef SPresetManager::FPresetViewEntry::EEntryType EEntryType;
+
+	public:
+
+		SLATE_BEGIN_ARGS(SCollectionTableRow) { }
+		    SLATE_ARGUMENT(TSharedPtr<SPresetManager::FPresetViewEntry>, ViewEntry)
+			SLATE_EVENT(FOnCollectionEnabledCheckboxChanged, OnCollectionEnabledCheckboxChanged)
+		SLATE_END_ARGS()
+
+		void Construct(const FArguments& InArgs, const TSharedRef<STableViewBase>& InOwnerTableView)
+		{
+			FTableRowArgs Args = FTableRowArgs()
+				.Style(&FAppStyle::Get().GetWidgetStyle<FTableRowStyle>("TableView.Row"))
+				.ExpanderStyleSet(& FCoreStyle::Get());
+
+			ViewEntry = InArgs._ViewEntry;
+			OnCollectionEnabledCheckboxChanged = InArgs._OnCollectionEnabledCheckboxChanged;
+
+			STableRow<ItemType>::Construct(Args, InOwnerTableView);
+		}
+
+		BEGIN_SLATE_FUNCTION_BUILD_OPTIMIZATION
+		virtual void ConstructChildren(ETableViewMode::Type InOwnerTableMode, const TAttribute<FMargin>& InPadding, const TSharedRef<SWidget>& InContent) override
+		{
+			STableRow<ItemType>::Content = InContent;
+
+			TSharedPtr<class ITableRow> ThisTableRow = this->SharedThis(this);
+
+			if(InOwnerTableMode == ETableViewMode::Tree)
+			{
+
+				// Rows in a TreeView need an expander button and some indentation
+				this->ChildSlot
+				[
+					SNew(SHorizontalBox)
+
+					+ SHorizontalBox::Slot()
+					.AutoWidth()
+					.HAlign(HAlign_Right)
+					.VAlign(VAlign_Fill)
+					[
+						SAssignNew(EnabledWidget, SCheckBox)
+						.Visibility_Lambda([this]() {return ViewEntry->EntryType == EEntryType::Collection ? EVisibility::Visible : EVisibility::Collapsed; })
+						.IsChecked_Lambda([this]() {return ViewEntry->bEnabled ? ECheckBoxState::Checked : ECheckBoxState::Unchecked; })
+						.OnCheckStateChanged_Lambda([this](ECheckBoxState State) {
+							if (OnCollectionEnabledCheckboxChanged.IsBound())
+							{
+								OnCollectionEnabledCheckboxChanged.Execute(ViewEntry, State);
+							}
+						})
+				
+					]
+
+					+ SHorizontalBox::Slot()
+					.AutoWidth()
+					.HAlign(HAlign_Right)
+					.VAlign(VAlign_Fill)
+					[
+						SAssignNew(STableRow< ItemType>::ExpanderArrowWidget, SExpanderArrow, ThisTableRow)
+						.StyleSet(STableRow< ItemType>::ExpanderStyleSet)
+						.ShouldDrawWires(false)
+					]
+
+					+ SHorizontalBox::Slot()
+						.AutoWidth()						
+						.Padding(5.0f)
+						[
+							SNew(SImage)
+							.Visibility_Lambda([this](){return ViewEntry->EntryType == EEntryType::Tool ? EVisibility::Visible : EVisibility::Collapsed; })
+							.Image(&ViewEntry->EntryIcon)
+						]
+
+					+ SHorizontalBox::Slot()
+						.FillWidth(1)
+						.Padding(5.0f)
+						[
+							SNew(STextBlock)
+							.Text(ViewEntry->EntryLabel)
+						]
+
+
+					+ SHorizontalBox::Slot()
+					.AutoWidth()
+					[
+						SNew(SSpacer)
+					
+					]
+
+					+ SHorizontalBox::Slot()
+						.AutoWidth()
+						.HAlign(EHorizontalAlignment::HAlign_Right)
+						.Padding(5.0f)
+						[
+							SNew(STextBlock)
+							.Text(TAttribute<FText>::CreateLambda([this]() { return FText::AsNumber(ViewEntry->Count); }))
+						]
+				];
+			}
+		}
+		END_SLATE_FUNCTION_BUILD_OPTIMIZATION
+
+	private:
+		TSharedPtr<SPresetManager::FPresetViewEntry> ViewEntry;
+		TSharedPtr<SCheckBox> EnabledWidget;
+
+		FOnCollectionEnabledCheckboxChanged OnCollectionEnabledCheckboxChanged;
+	};
+
+
+
+	template<typename ItemType>
+	class SPresetTableRow : public SMultiColumnTableRow< ItemType>
+	{
+		typedef SPresetManager::FPresetViewEntry::EEntryType EEntryType;
+		using typename SMultiColumnTableRow< ItemType>::FSuperRowType;
+
+	public:
+
+		SLATE_BEGIN_ARGS(SPresetTableRow) { }
+		SLATE_ARGUMENT(TSharedPtr<SPresetManager::FPresetViewEntry>, ViewEntry)
+		SLATE_EVENT(FOnPresetDeleted, OnPresetDeleted)
+		SLATE_END_ARGS()
+
+		void Construct(const FArguments& InArgs, const TSharedRef<STableViewBase>& InOwnerTableView)
+		{
+			typename FSuperRowType::FArguments Args = typename FSuperRowType::FArguments()
+				.ExpanderStyleSet(&FCoreStyle::Get());
+
+			OnPresetDeleted = InArgs._OnPresetDeleted;
+			ViewEntry = InArgs._ViewEntry;
+
+			SMultiColumnTableRow<ItemType>::Construct(Args, InOwnerTableView);
+		}
+
+		BEGIN_SLATE_FUNCTION_BUILD_OPTIMIZATION
+		virtual void ConstructChildren(ETableViewMode::Type InOwnerTableMode, const TAttribute<FMargin>& InPadding, const TSharedRef<SWidget>& InContent) override
+		{
+			STableRow<ItemType>::Content = InContent;
+
+			TSharedPtr<class ITableRow> ThisTableRow = this->SharedThis(this);
+
+			if (ViewEntry->EntryType == EEntryType::Tool)
+			{
+				this->ChildSlot				
+				[
+				SNew(SHorizontalBox)
+
+				+ SHorizontalBox::Slot()
+					.AutoWidth()
+					.Padding(5.0f)
+					[
+						SNew(SImage)
+						.Visibility_Lambda([this]() {return ViewEntry->EntryType == EEntryType::Tool ? EVisibility::Visible : EVisibility::Collapsed; })
+					.Image(&ViewEntry->EntryIcon)
+					]
+
+				+ SHorizontalBox::Slot()
+					.FillWidth(1)
+					.Padding(5.0f)
+					[
+						SNew(STextBlock)
+						.Text(ViewEntry->EntryLabel)
+					]
+
+				+ SHorizontalBox::Slot()
+					.AutoWidth()
+					[
+						SNew(SSpacer)
+
+					]
+
+				+ SHorizontalBox::Slot()
+					.AutoWidth()
+					.HAlign(EHorizontalAlignment::HAlign_Right)
+					.Padding(5.0f)
+					[
+						SNew(STextBlock)
+						.Text(TAttribute<FText>::CreateLambda([this]() { return FText::AsNumber(ViewEntry->Count); }))
+					]
+				];
+			}
+			else
+			{
+				ensure(ViewEntry->EntryType == EEntryType::Preset);
+
+				this->ChildSlot
+				.Padding(InPadding)
+				[
+					InContent
+				];
+			}
+			
+		}
+
+		virtual TSharedRef<SWidget> GenerateWidgetForColumn(const FName& InColumnName)
+		{
+			if (InColumnName == "Label")
+			{
+				return SNew(SHorizontalBox)
+				+ SHorizontalBox::Slot()
+					.VAlign(VAlign_Center)
+					.Padding(FMargin(4.0f, 0.0f))
+					[
+						SNew(STextBlock)
+						.Text(TAttribute<FText>::CreateLambda([this]() { return FText::FromString(ViewEntry->PresetLabel); }))
+					];
+			}
+			else if (InColumnName == "Tooltip")
+			{
+				return SNew(SHorizontalBox)
+				+ SHorizontalBox::Slot()
+					.VAlign(VAlign_Center)
+					.Padding(FMargin(4.0f, 0.0f))
+					[
+						SNew(STextBlock)
+						.Text(TAttribute<FText>::CreateLambda([this]() { return FText::FromString(ViewEntry->PresetTooltip); }))
+					];
+			}
+			else if (InColumnName == "Tool")
+			{
+				return SNew(SBox)
+					.VAlign(VAlign_Center)
+					.Padding(FMargin(5.0f, 5.0f))
+					[
+						SNew(SImage)
+						.Image(&ViewEntry->EntryIcon)
+						.DesiredSizeOverride(FVector2D(16, 16))
+					];
+			}
+			else if (InColumnName == "Delete")
+			{
+				return SNew(SNegativeActionButton)
+					.Icon(FAppStyle::GetBrush("Icons.Delete"))						
+					.OnClicked_Lambda([this]() { OnPresetDeleted.ExecuteIfBound(ViewEntry); return FReply::Handled(); })
+					.Visibility_Lambda([this]() { return this->IsHovered() ? EVisibility::Visible : EVisibility::Hidden; });
+			}
+
+			return SNullWidget::NullWidget;
+
+		}
+		END_SLATE_FUNCTION_BUILD_OPTIMIZATION
+
+	private:
+		TSharedPtr<SPresetManager::FPresetViewEntry> ViewEntry;
+		FOnPresetDeleted                     OnPresetDeleted;
+	};
+
+}
+
+
+
+/* SPresetManager interface
+ *****************************************************************************/
+
+BEGIN_SLATE_FUNCTION_BUILD_OPTIMIZATION
+void SPresetManager::Construct( const FArguments& InArgs )
+{
+	UPresetUserSettings::Initialize();
+	UserSettings = UPresetUserSettings::Get();
+	if (UserSettings.IsValid())
+	{
+		UserSettings->LoadEditorConfig();
+	}
+
+	ChildSlot
+	[
+		SNew(SVerticalBox)
+
+		+ SVerticalBox::Slot()
+		.AutoHeight()
+		[
+			SNew(SHorizontalBox)
+
+			+ SHorizontalBox::Slot()
+			.AutoWidth()
+			.Padding(16.0f, 8.0f, 8.0f, 8.0f)
+			//.HAlign(HAlign_Fill)
+			//.SizeParam(FStretch(1.0f))
+			[
+				SAssignNew(AddUserPresetButton, SPositiveActionButton)
+				.Text(LOCTEXT( "AddUserPresetCollection", "Add User Preset Collection") )
+				.Icon(FAppStyle::GetBrush("Icons.PlusCircle"))
+				.OnClicked_Lambda([this]() { AddNewUserPresetCollection(); return FReply::Handled(); })
+			]
+
+			+ SHorizontalBox::Slot()
+			.AutoWidth()
+			.Padding(8.0f, 8.0f, 8.0f, 8.0f)
+			//.HAlign(HAlign_Fill)
+			//.SizeParam(FStretch(1.0f))
+			[
+				SAssignNew(DeleteUserPresetButton, SNegativeActionButton)
+				.Text( LOCTEXT( "DeleteUserPresetCollection", "Delete User Preset Collection"))
+				.Icon(FAppStyle::GetBrush("Icons.Delete"))
+				.OnClicked_Lambda([this]() { DeleteSelectedUserPresetCollection(); return FReply::Handled(); })
+			]
+
+		]
+		+ SVerticalBox::Slot()
+		.FillHeight(1.0)
+		//.AutoHeight()
+		[
+				SAssignNew(Splitter, SSplitter)
+					.Orientation(EOrientation::Orient_Horizontal)
+					
+					+ SSplitter::Slot()
+					.Value(0.4f)
+					.Resizable(false)
+						[
+							SNew(SVerticalBox)												
+								+ SVerticalBox::Slot()
+								.AutoHeight()
+								[
+									SNew(SBorder)
+									.BorderImage(&FAppStyle::Get().GetWidgetStyle<FTableColumnHeaderStyle>("TableView.Header.Column").NormalBrush)
+									[
+										SNew(SHorizontalBox)
+									
+										+ SHorizontalBox::Slot()
+										.AutoWidth()
+										.HAlign(HAlign_Right)
+										.VAlign(VAlign_Fill)
+										[
+											SAssignNew(UserCollectionsExpander, SButton)
+											.ButtonStyle(FCoreStyle::Get(), "NoBorder" )
+											.VAlign(VAlign_Center)
+											.HAlign(HAlign_Center)
+											.ClickMethod( EButtonClickMethod::MouseDown )
+											.OnClicked_Lambda([this]() { bAreUserCollectionsExpanded = !bAreUserCollectionsExpanded; return FReply::Handled(); })
+											.ContentPadding(0.f)
+											.ForegroundColor( FSlateColor::UseForeground() )
+											.IsFocusable( false )
+											[
+												SNew(SImage)
+												.Image( this, &SPresetManager::GetUserCollectionsExpanderImage )
+												.ColorAndOpacity( FSlateColor::UseSubduedForeground() )
+											]
+										]
+									
+
+										+ SHorizontalBox::Slot()
+										.AutoWidth()
+										.Padding(5.0f, 5.0f, 5.0f, 5.0f)
+										.HAlign(HAlign_Center)
+										.VAlign(VAlign_Center)
+										[
+											SNew(STextBlock)										
+											.Text(LOCTEXT("UserPresetLabels", "User Presets"))
+											.Font(FSlateFontInfo(FCoreStyle::GetDefaultFont(), 13, "Bold"))
+										]
+									]
+								]
+
+								+ SVerticalBox::Slot()
+									.FillHeight(1.0f)
+									//.AutoHeight()									
+								[
+			
+											SAssignNew(UserPresetCollectionTreeView, STreeView<TSharedPtr<FPresetViewEntry> >)
+												.ItemHeight(32.0f)
+												.TreeItemsSource(&UserCollectionsDataList)
+												.SelectionMode(ESelectionMode::Single)
+												.OnGenerateRow(this, &SPresetManager::HandleTreeGenerateRow)
+												.OnGetChildren(this, &SPresetManager::HandleTreeGetChildren)
+												.OnSelectionChanged(this, &SPresetManager::HandleUserTreeSelectionChanged)
+												.Visibility_Lambda([this]() {return bAreUserCollectionsExpanded ? EVisibility::Visible : EVisibility::Collapsed; })
+												.HeaderRow
+												(
+													SNew(SHeaderRow)
+													.Visibility(EVisibility::Collapsed)
+
+													+ SHeaderRow::Column("Collection")
+													.FixedWidth(150.0f)
+													.HeaderContent()
+													[
+														SNew(STextBlock)
+														.Text(LOCTEXT("PresetManagerCollectionTitleHeader", "Collection"))
+													]
+
+
+												)											
+								]
+
+								+ SVerticalBox::Slot()
+									.AutoHeight()									
+								[
+									SNew(SBorder)
+									.BorderImage(&FAppStyle::Get().GetWidgetStyle<FTableColumnHeaderStyle>("TableView.Header.Column").NormalBrush)
+									[
+										SNew(SHorizontalBox)
+										+ SHorizontalBox::Slot()
+										.AutoWidth()
+										.HAlign(HAlign_Right)
+										.VAlign(VAlign_Fill)
+										[
+											SAssignNew(ProjectCollectionsExpander, SButton)
+											.ButtonStyle(FCoreStyle::Get(), "NoBorder" )
+											.VAlign(VAlign_Center)
+											.HAlign(HAlign_Center)
+											.ClickMethod( EButtonClickMethod::MouseDown )
+											.OnClicked_Lambda([this]() { bAreProjectCollectionsExpanded = !bAreProjectCollectionsExpanded; return FReply::Handled(); })
+											.ContentPadding(0.f)
+											.ForegroundColor( FSlateColor::UseForeground() )
+											.IsFocusable( false )
+											[
+												SNew(SImage)
+												.Image( this, &SPresetManager::GetProjectCollectionsExpanderImage )
+												.ColorAndOpacity( FSlateColor::UseSubduedForeground() )
+											]
+										]
+
+										+ SHorizontalBox::Slot()
+										.AutoWidth()
+										.Padding(5.0f, 5.0f, 5.0f, 5.0f)
+										.HAlign(HAlign_Center)
+										.VAlign(VAlign_Center)
+										[
+											SNew(STextBlock)
+											.Text(LOCTEXT("ProjectPresetLabels", "Project Presets"))
+											.Font(FSlateFontInfo(FCoreStyle::GetDefaultFont(), 13, "Bold"))
+										]
+									]
+								]
+
+								+ SVerticalBox::Slot()
+									.FillHeight(1.0f)
+									//.AutoHeight()
+								[
+										
+											SNew(SVerticalBox)
+											.Visibility_Lambda([this]() {return bAreProjectCollectionsExpanded ? EVisibility::Visible : EVisibility::Collapsed; })
+
+											+ SVerticalBox::Slot()
+											.FillHeight(1.0f)
+											[
+
+											SAssignNew(ProjectPresetCollectionTreeView, STreeView<TSharedPtr<FPresetViewEntry> >)
+											.Visibility(this, &SPresetManager::ProjectPresetCollectionsVisibility)
+												.ItemHeight(32.0f)
+												.TreeItemsSource(&ProjectCollectionsDataList)
+												.SelectionMode(ESelectionMode::Single)
+												.OnGenerateRow(this, &SPresetManager::HandleTreeGenerateRow)
+												.OnGetChildren(this, &SPresetManager::HandleTreeGetChildren)
+												.OnSelectionChanged(this, &SPresetManager::HandleTreeSelectionChanged)
+												.HeaderRow
+												(
+													SNew(SHeaderRow)
+													.Visibility(EVisibility::Collapsed)
+
+													+ SHeaderRow::Column("Collection")
+													.FixedWidth(150.0f)
+													.HeaderContent()
+													[
+														SNew(STextBlock)
+														.Text(LOCTEXT("PresetManagerCollectionTitleHeader", "Collection"))
+													]
+
+
+												)											
+											]
+
+											+ SVerticalBox::Slot()
+											.AutoHeight()
+											.HAlign(HAlign_Center)
+											[
+												SNew(STextBlock)
+												.Visibility_Lambda([this]() { return ProjectPresetCollectionsVisibility() == EVisibility::Visible ? EVisibility::Collapsed : EVisibility::Visible; })
+												.Text(LOCTEXT("ProjectPresetsNotLoadedLabel", "(Load Project Preset Collections in Project Settings)"))
+												.Font(FAppStyle::GetFontStyle("NormalFontItalic"))
+											]
+
+
+								]
+
+						]
+
+						+ SSplitter::Slot()
+						[
+							SNew(SVerticalBox)
+
+							+ SVerticalBox::Slot()
+							.FillHeight(1.0f)
+							[
+
+								SAssignNew(PresetListView, SListView<TSharedPtr<FPresetViewEntry>>)
+								.ListItemsSource(&PresetDataList)
+								.ItemHeight(32.0f)
+								.SelectionMode(ESelectionMode::SingleToggle)
+								.OnGenerateRow(this, &SPresetManager::HandleListGenerateRow)		
+								.OnSelectionChanged(this, &SPresetManager::HandleListSelectionChanged)
+								.HeaderRow
+								(
+									SNew(SHeaderRow)
+									.Visibility(EVisibility::Visible)
+
+									+ SHeaderRow::Column("Tool")
+									.FixedWidth(30.0f)
+									.HeaderContentPadding(FMargin(5.0f, 5.0f))
+									.HAlignHeader(EHorizontalAlignment::HAlign_Center)
+									.VAlignHeader(EVerticalAlignment::VAlign_Center)
+									.HAlignCell(EHorizontalAlignment::HAlign_Center)
+									.HeaderContent()									
+									[
+										SNew(SImage)
+										.Image(FPresetEditorStyle::Get()->GetBrush("ManagerIcons.Tools"))
+										.DesiredSizeOverride(FVector2D(20, 20))
+									]
+									
+									+ SHeaderRow::Column("Label")
+									.FillWidth(80.0f)
+									.HeaderContentPadding(FMargin(5.0f, 5.0f))
+									.HAlignHeader(EHorizontalAlignment::HAlign_Left)
+									.VAlignHeader(EVerticalAlignment::VAlign_Center)
+									.HeaderContent()
+									[
+										SNew(STextBlock)
+										.Text(LOCTEXT("PresetManagerPresetLabelHeader", "Label"))
+									]
+
+									+ SHeaderRow::Column("Tooltip")
+									.FillWidth(80.0f)
+									.HAlignHeader(EHorizontalAlignment::HAlign_Left)
+									.VAlignHeader(EVerticalAlignment::VAlign_Center)
+									.HeaderContent()
+									[
+										SNew(STextBlock)
+										.Text(LOCTEXT("PresetManagerPresetTooltipHeader", "Tooltip"))
+									]
+
+									+ SHeaderRow::Column("Delete")
+									.FixedWidth(60.0f)									
+									.HAlignHeader(EHorizontalAlignment::HAlign_Center)
+									.VAlignHeader(EVerticalAlignment::VAlign_Center)
+									.HAlignCell(EHorizontalAlignment::HAlign_Center)
+									.HeaderContent()
+									[
+										SNew(SImage)
+										.Image(FAppStyle::GetBrush("Icons.Delete"))
+									]
+
+								)
+						]
+
+						+ SVerticalBox::Slot()
+						.AutoHeight()
+						.Padding(5.0f)
+						[
+							SNew(SVerticalBox)
+							.Visibility(this, &SPresetManager::EditAreaVisibility)
+							+ SVerticalBox::Slot()
+							.AutoHeight()
+							.Padding(5.0f)
+							[
+								SNew(SHorizontalBox)								
+								+ SHorizontalBox::Slot()
+							    .FillWidth(1.f)
+								.Padding(5.0f)
+								.HAlign(EHorizontalAlignment::HAlign_Left)	
+								[
+									SNew(STextBlock)
+									.Text(LOCTEXT("PresetLabelEditLabel", "Label"))
+								]
+								+ SHorizontalBox::Slot()
+								.FillWidth(1.f)								
+								.Padding(5.0f)
+								[
+									SNew(SEditableTextBox)			
+									.Text_Lambda([this]() {
+										if (ActivePresetToEdit)
+										{
+											return FText::FromString(ActivePresetToEdit->PresetLabel);
+										}
+										return FText::GetEmpty();
+									})
+									.OnTextChanged_Lambda([this](const FText& NewText) {
+										if (ActivePresetToEdit)
+										{
+											ActivePresetToEdit->PresetLabel = NewText.ToString();
+										}
+									})
+									.OnTextCommitted_Lambda([this](const FText& NewText, ETextCommit::Type CommitStatus){
+										if (ActivePresetToEdit)
+										{
+											SetPresetLabel(ActivePresetToEdit, NewText);
+										}
+									})
+								]
+							]
+							+ SVerticalBox::Slot()
+							.AutoHeight()
+							.Padding(5.0f)
+							[
+								SNew(SHorizontalBox)
+								+ SHorizontalBox::Slot()
+								.FillWidth(1.f)
+								.Padding(5.0f)
+								.HAlign(EHorizontalAlignment::HAlign_Left)
+								[
+									SNew(STextBlock)
+									.Text(LOCTEXT("PresetTooltipEditLabel", "Tooltip"))
+								]
+								+ SHorizontalBox::Slot()
+								.FillWidth(1.f)
+								.Padding(5.0f)																
+								[
+									SNew(SBox)
+									.MinDesiredHeight(44.f)
+									.MaxDesiredHeight(44.0f)
+									[
+										SNew(SMultiLineEditableTextBox)											
+										.AllowMultiLine(false)
+										.AutoWrapText(true)
+										.WrappingPolicy(ETextWrappingPolicy::DefaultWrapping)
+										.Text_Lambda([this]() {
+											if (ActivePresetToEdit)
+											{
+												return FText::FromString(ActivePresetToEdit->PresetTooltip);
+											}
+											return FText::GetEmpty();
+										})
+										.OnTextChanged_Lambda([this](const FText& NewText) {
+											if (ActivePresetToEdit)
+											{
+												ActivePresetToEdit->PresetTooltip = NewText.ToString();
+											}
+										})
+										.OnTextCommitted_Lambda([this](const FText& NewText, ETextCommit::Type CommitStatus){
+											if (ActivePresetToEdit)
+											{
+												SetPresetTooltip(ActivePresetToEdit, NewText);
+											}
+										})
+									]
+								]
+							]
+						]						
+					]
+				]
+		];
+
+		DeleteUserPresetButton->SetEnabled(false);
+}
+END_SLATE_FUNCTION_BUILD_OPTIMIZATION
+
+SPresetManager::~SPresetManager()
+{
+}
+
+void SPresetManager::Tick(const FGeometry& AllottedGeometry, const double InCurrentTime, const float InDeltaTime)
+{
+	if (!ensure(UserSettings.IsValid()))
+	{
+		return;
+	}
+
+	const UPresetProjectSettings* ProjectSettings = GetDefault<UPresetProjectSettings>();
+	TArray<FSoftObjectPath> AvailablePresetCollections = ProjectSettings->LoadedPresetCollections.Array();
+	TArray<FSoftObjectPath> AvailableUserPresetCollections;
+	PresetManagerLocals::GetObjectsOfClass<UInteractiveToolsPresetCollectionAsset>(AvailableUserPresetCollections);
+
+	TotalPresetCount = 0;
+
+	auto GenerateTreeEntries = [this](TArray<FSoftObjectPath>& AssetList,
+		TArray< TSharedPtr< FPresetViewEntry > >& TreeList,
+		TSharedPtr<STreeView<TSharedPtr<FPresetViewEntry> > >& TreeView)
+	{
+		AssetList.RemoveAll([](const FSoftObjectPath& Path) { return !Path.IsAsset(); });
+		
+		bool bTreeNeedsRefresh = false;
+		TArray< TSharedPtr< FPresetViewEntry > > TempTreeDataList;
+		for (const FSoftObjectPath& Path : AssetList)
+		{
+			UInteractiveToolsPresetCollectionAsset* PresetCollection = nullptr;
+
+			if (Path.IsAsset())
+			{
+				PresetCollection = Cast<UInteractiveToolsPresetCollectionAsset>(Path.TryLoad());
+			}
+			if (PresetCollection)
+			{
+				if (Path == UserSettings->DefaultCollectionPath)
+				{
+					TempTreeDataList.Add(MakeShared<FPresetViewEntry>(
+						UserSettings->bDefaultCollectionEnabled,
+						Path,
+						PresetCollection->CollectionLabel,
+						0));
+				}
+				else
+				{
+					TempTreeDataList.Add(MakeShared<FPresetViewEntry>(
+						UserSettings->EnabledPresetCollections.Contains(Path),
+						Path,
+						PresetCollection->CollectionLabel,
+						0));
+				}
+
+				TMap<FString, FInteractiveToolPresetStore >::TIterator ToolNameIter = PresetCollection->PerToolPresets.CreateIterator();
+				for (; (bool)ToolNameIter; ++ToolNameIter)
+				{
+					int32 ToolCount = 0;
+					for (int32 PresetIndex = 0; PresetIndex < ToolNameIter.Value().NamedPresets.Num(); ++PresetIndex)
+					{
+						ToolCount += ToolNameIter.Value().NamedPresets[PresetIndex].IsValid() ? 1 : 0;
+					}
+
+					TempTreeDataList.Last()->Children.Add(MakeShared<FPresetViewEntry>(
+						ToolNameIter.Value().ToolLabel,
+						ToolNameIter.Value().ToolIcon,
+						Path,
+						ToolNameIter.Key(),
+						ToolCount));
+					TempTreeDataList.Last()->Children.Last()->Parent = TempTreeDataList.Last();
+					TempTreeDataList.Last()->Count += ToolCount;
+					TotalPresetCount += ToolCount;
+				}
+			}
+		}
+
+		if (TempTreeDataList.Num() != TreeList.Num())
+		{
+			bTreeNeedsRefresh = true;
+		}
+		else
+		{
+			for (int32 CollectionIndex = 0; CollectionIndex < TreeList.Num(); ++CollectionIndex)
+			{
+				if (!(TreeList[CollectionIndex]->HasSameMetadata(*TempTreeDataList[CollectionIndex])))
+				{
+					bTreeNeedsRefresh = true;
+				}
+			}
+		}
+
+		if (bTreeNeedsRefresh)
+		{
+			TreeList = TempTreeDataList;
+			TreeView->RequestTreeRefresh();
+			bHasActiveCollection = false;
+		}
+
+		for (TSharedPtr<FPresetViewEntry>& Entry : TreeList)
+		{
+			Entry->bEnabled = UserSettings->EnabledPresetCollections.Contains(Entry->CollectionPath);
+			if (Entry->CollectionPath == UserSettings->DefaultCollectionPath)
+			{
+				Entry->bEnabled = UserSettings->bDefaultCollectionEnabled;
+			}
+			else
+			{
+				Entry->bEnabled = UserSettings->EnabledPresetCollections.Contains(Entry->CollectionPath);
+			}
+		}
+	};
+
+	GenerateTreeEntries(AvailablePresetCollections, ProjectCollectionsDataList, ProjectPresetCollectionTreeView);
+	GenerateTreeEntries(AvailableUserPresetCollections, UserCollectionsDataList, UserPresetCollectionTreeView);
+
+}
+
+
+/* SPresetManager implementation
+ *****************************************************************************/
+
+int32 SPresetManager::GetTotalPresetCount() const
+{
+	return TotalPresetCount;
+}
+
+TSharedRef<ITableRow> SPresetManager::HandleTreeGenerateRow(TSharedPtr<FPresetViewEntry> TreeEntry, const TSharedRef<STableViewBase>& OwnerTable)
+{
+	return SNew(PresetManagerLocals::SCollectionTableRow< TSharedPtr<FPresetViewEntry> >, OwnerTable)
+		.ViewEntry(TreeEntry)
+		.OnCollectionEnabledCheckboxChanged(this, &SPresetManager::SetCollectionEnabled);
+}
+
+void SPresetManager::HandleTreeGetChildren(TSharedPtr<FPresetViewEntry> TreeEntry, TArray< TSharedPtr<FPresetViewEntry> >& ChildrenOut)
+{
+	ChildrenOut = TreeEntry->Children;
+}
+
+void SPresetManager::GeneratePresetList(TSharedPtr<FPresetViewEntry> TreeEntry)
+{
+	PresetDataList.Empty();
+	PresetListView->RequestListRefresh();
+	bHasActiveCollection = false;
+
+	if (!TreeEntry)
+	{
+		return;
+	}
+
+	if (TreeEntry->EntryType == FPresetViewEntry::EEntryType::Collection ||
+		TreeEntry->EntryType == FPresetViewEntry::EEntryType::Tool)
+	{
+		UInteractiveToolsPresetCollectionAsset* PresetCollection = nullptr;
+
+		if (TreeEntry->CollectionPath.IsAsset())
+		{
+			PresetCollection = Cast<UInteractiveToolsPresetCollectionAsset>(TreeEntry->CollectionPath.TryLoad());
+		}
+		if (PresetCollection)
+		{
+			if (TreeEntry->EntryType == FPresetViewEntry::EEntryType::Collection)
+			{
+				bHasActiveCollection = true;
+				bIsActiveCollectionEnabled = TreeEntry->bEnabled;
+				ActiveCollectionLabel = TreeEntry->EntryLabel;
+
+				TMap<FString, FInteractiveToolPresetStore >::TIterator ToolNameIter = PresetCollection->PerToolPresets.CreateIterator();
+				for (; (bool)ToolNameIter; ++ToolNameIter)
+				{
+					int32 ToolCount = ToolNameIter.Value().NamedPresets.Num();
+					for (int32 PresetIndex = 0; PresetIndex < ToolNameIter.Value().NamedPresets.Num(); ++PresetIndex)
+					{
+						if (ToolNameIter.Value().NamedPresets[PresetIndex].IsValid())
+						{
+							PresetDataList.Add(MakeShared<FPresetViewEntry>(
+								ToolNameIter.Key(),
+								PresetIndex,
+								ToolNameIter.Value().NamedPresets[PresetIndex].Label,
+								ToolNameIter.Value().NamedPresets[PresetIndex].Tooltip,
+								FText::FromString(ToolNameIter.Value().NamedPresets[PresetIndex].Label)
+								));
+							PresetDataList.Last()->CollectionPath = TreeEntry->CollectionPath;
+							PresetDataList.Last()->EntryIcon = ToolNameIter.Value().ToolIcon;
+						}
+					}
+				}
+			}
+			else
+			{
+				bHasActiveCollection = true;
+				bIsActiveCollectionEnabled = TreeEntry->Parent->bEnabled;
+				ActiveCollectionLabel = TreeEntry->Parent->EntryLabel;
+
+				const FInteractiveToolPresetStore* ToolData = PresetCollection->PerToolPresets.Find(TreeEntry->ToolName);
+				if (!ToolData)
+				{
+					return;
+				}
+				int32 ToolCount = ToolData->NamedPresets.Num();
+				for (int32 PresetIndex = 0; PresetIndex < ToolData->NamedPresets.Num(); ++PresetIndex)
+				{
+					if (ToolData->NamedPresets[PresetIndex].IsValid())
+					{
+						PresetDataList.Add(MakeShared<FPresetViewEntry>(
+							TreeEntry->ToolName,
+							PresetIndex,
+							ToolData->NamedPresets[PresetIndex].Label,
+							ToolData->NamedPresets[PresetIndex].Tooltip,
+							FText::FromString(ToolData->NamedPresets[PresetIndex].Label)
+							));
+						PresetDataList.Last()->CollectionPath = TreeEntry->CollectionPath;
+						PresetDataList.Last()->EntryIcon = TreeEntry->EntryIcon;
+					}
+				}
+			}
+		}
+	}
+
+}
+
+
+void SPresetManager::HandleTreeSelectionChanged(TSharedPtr<FPresetViewEntry> TreeEntry, ESelectInfo::Type SelectInfo)
+{
+	if (SelectInfo != ESelectInfo::Direct) {
+		UserPresetCollectionTreeView->ClearSelection();
+		GeneratePresetList(TreeEntry);
+	}
+}
+
+void SPresetManager::HandleUserTreeSelectionChanged(TSharedPtr<FPresetViewEntry> TreeEntry, ESelectInfo::Type SelectInfo)
+{
+	if (SelectInfo != ESelectInfo::Direct) {
+		ProjectPresetCollectionTreeView->ClearSelection();
+		GeneratePresetList(TreeEntry);
+	}
+
+	DeleteUserPresetButton->SetEnabled(false);
+	if (TreeEntry)
+	{
+		if (TreeEntry->EntryType == FPresetViewEntry::EEntryType::Collection && TreeEntry->CollectionPath != UserSettings->DefaultCollectionPath)
+		{
+			DeleteUserPresetButton->SetEnabled(true);
+		}
+	}
+
+}
+
+TSharedRef<ITableRow> SPresetManager::HandleListGenerateRow(TSharedPtr<FPresetViewEntry> TreeEntry, const TSharedRef<STableViewBase>& OwnerTable)
+{
+	return SNew(PresetManagerLocals::SPresetTableRow< TSharedPtr<FPresetViewEntry> >, OwnerTable)
+		.ViewEntry(TreeEntry)
+		.OnPresetDeleted(this, &SPresetManager::DeletePresetFromCollection);
+}
+
+void SPresetManager::HandleListSelectionChanged(TSharedPtr<FPresetViewEntry> TreeEntry, ESelectInfo::Type SelectInfo)
+{
+	if (TreeEntry)
+	{
+		ActivePresetToEdit = TreeEntry;
+	}
+	else
+	{
+		if (ActivePresetToEdit)
+		{
+			SetPresetLabel(ActivePresetToEdit, FText::FromString(ActivePresetToEdit->PresetLabel));
+			SetPresetTooltip(ActivePresetToEdit, FText::FromString(ActivePresetToEdit->PresetTooltip));
+		}
+		ActivePresetToEdit.Reset();
+	}
+}
+
+EVisibility SPresetManager::EditAreaVisibility() const
+{
+	return ActivePresetToEdit.IsValid() ? EVisibility::Visible : EVisibility::Collapsed;
+}
+
+EVisibility SPresetManager::ProjectPresetCollectionsVisibility() const
+{
+	return ProjectCollectionsDataList.IsEmpty() ? EVisibility::Collapsed : EVisibility::Visible;
+}
+
+void SPresetManager::SetCollectionEnabled(TSharedPtr<FPresetViewEntry> TreeEntry, ECheckBoxState State)
+{	
+	if (!ensure(UserSettings.IsValid()))
+	{
+		return;
+	}
+	if (UserSettings->DefaultCollectionPath == TreeEntry->CollectionPath)
+	{
+		UserSettings->bDefaultCollectionEnabled = (State == ECheckBoxState::Checked);
+		UserSettings->SaveEditorConfig();
+	}
+	else
+	{
+		if (State == ECheckBoxState::Checked && !UserSettings->EnabledPresetCollections.Contains(TreeEntry->CollectionPath))
+		{
+			UserSettings->EnabledPresetCollections.Add(TreeEntry->CollectionPath);
+			UserSettings->SaveEditorConfig();
+		}
+		else if (State != ECheckBoxState::Checked && UserSettings->EnabledPresetCollections.Contains(TreeEntry->CollectionPath))
+		{
+			UserSettings->EnabledPresetCollections.Remove(TreeEntry->CollectionPath);
+			UserSettings->SaveEditorConfig();
+		}
+	}
+}
+
+void SPresetManager::DeletePresetFromCollection(TSharedPtr< FPresetViewEntry > Entry)
+{
+	UInteractiveToolsPresetCollectionAsset* PresetCollection = nullptr;
+
+	if (Entry->CollectionPath.IsAsset())
+	{
+		PresetCollection = Cast<UInteractiveToolsPresetCollectionAsset>(Entry->CollectionPath.TryLoad());
+	}
+	if (PresetCollection)
+	{
+		PresetCollection->PerToolPresets[Entry->ToolName].NamedPresets.RemoveAt(Entry->PresetIndex);
+		PresetCollection->MarkPackageDirty();
+
+		PresetDataList.Remove(Entry);
+		PresetListView->RequestListRefresh();
+	}
+}
+
+void SPresetManager::SetPresetLabel(TSharedPtr< FPresetViewEntry > Entry, FText InLabel)
+{
+	UInteractiveToolsPresetCollectionAsset* PresetCollection = nullptr;
+
+	if (Entry->CollectionPath.IsAsset())
+	{
+		PresetCollection = Cast<UInteractiveToolsPresetCollectionAsset>(Entry->CollectionPath.TryLoad());
+	}
+	if (PresetCollection)
+	{
+		PresetCollection->PerToolPresets[Entry->ToolName].NamedPresets[Entry->PresetIndex].Label = InLabel.ToString();
+		PresetCollection->MarkPackageDirty();
+	}
+}
+
+void SPresetManager::SetPresetTooltip(TSharedPtr< FPresetViewEntry > Entry, FText InTooltip)
+{
+	UInteractiveToolsPresetCollectionAsset* PresetCollection = nullptr;
+
+	if (Entry->CollectionPath.IsAsset())
+	{
+		PresetCollection = Cast<UInteractiveToolsPresetCollectionAsset>(Entry->CollectionPath.TryLoad());
+	}
+	if (PresetCollection)
+	{
+		PresetCollection->PerToolPresets[Entry->ToolName].NamedPresets[Entry->PresetIndex].Tooltip = InTooltip.ToString();
+		PresetCollection->MarkPackageDirty();
+	}
+}
+
+void SPresetManager::DeleteSelectedUserPresetCollection()
+{
+	FAssetRegistryModule& AssetRegistryModule = FModuleManager::LoadModuleChecked<FAssetRegistryModule>(TEXT("AssetRegistry"));
+	FContentBrowserModule& ContentBrowserModule = FModuleManager::LoadModuleChecked<FContentBrowserModule>("ContentBrowser");
+
+	TArray<TSharedPtr<FPresetViewEntry>> SelectedUserCollections = UserPresetCollectionTreeView->GetSelectedItems();
+
+	if (SelectedUserCollections.Num() == 1)
+	{
+		TSharedPtr<FPresetViewEntry> Entry = SelectedUserCollections[0];
+		FAssetData CollectionAsset;
+		if (AssetRegistryModule.Get().TryGetAssetByObjectPath(Entry->CollectionPath, CollectionAsset) == UE::AssetRegistry::EExists::Exists)
+		{
+			TArray<FAssetData> AssetData;
+			AssetData.Add(CollectionAsset);
+			ObjectTools::DeleteAssets(AssetData, true);
+		}
+	}
+}
+
+void SPresetManager::AddNewUserPresetCollection()
+{
+	// Load necessary modules
+	FAssetToolsModule& AssetToolsModule = FModuleManager::Get().LoadModuleChecked<FAssetToolsModule>("AssetTools");
+	FContentBrowserModule& ContentBrowserModule = FModuleManager::LoadModuleChecked<FContentBrowserModule>("ContentBrowser");
+	FAssetRegistryModule& AssetRegistryModule = FModuleManager::LoadModuleChecked<FAssetRegistryModule>(TEXT("AssetRegistry"));
+	IAssetRegistry& AssetRegistry = AssetRegistryModule.Get();
+
+	// Generate a unique asset name
+	FString Name, PackageName;
+	AssetToolsModule.Get().CreateUniqueAssetName(TEXT("/ToolPresets/Presets/"), TEXT("UserPresetCollection"), PackageName, Name);
+	const FString PackagePath = FPackageName::GetLongPackagePath(PackageName);
+
+	// Create object and package
+	UPackage* package = CreatePackage(*PackageName);
+	UInteractiveToolsPresetCollectionAssetFactory* MyFactory = NewObject<UInteractiveToolsPresetCollectionAssetFactory>(UInteractiveToolsPresetCollectionAssetFactory::StaticClass()); // Can omit, and a default factory will be used
+	UObject* NewObject = AssetToolsModule.Get().CreateAsset(Name, PackagePath, UInteractiveToolsPresetCollectionAsset::StaticClass(), MyFactory);
+	UInteractiveToolsPresetCollectionAsset* NewCollection = ExactCast<UInteractiveToolsPresetCollectionAsset>(NewObject);
+	NewCollection->CollectionLabel = FText::FromString(Name);
+	FSavePackageArgs SavePackageArgs;
+	SavePackageArgs.TopLevelFlags = RF_Public | RF_Standalone;
+	UPackage::Save(package, NewObject, *FPackageName::LongPackageNameToFilename(PackageName, FPackageName::GetAssetPackageExtension()), SavePackageArgs);
+
+	// Inform asset registry
+	AssetRegistry.AssetCreated(NewObject);
+}
+
+const FSlateBrush* SPresetManager::GetProjectCollectionsExpanderImage() const
+{
+	return GetExpanderImage(ProjectCollectionsExpander, false);
+}
+
+const FSlateBrush* SPresetManager::GetUserCollectionsExpanderImage() const
+{
+	return GetExpanderImage(UserCollectionsExpander, true);
+}
+
+const FSlateBrush* SPresetManager::GetExpanderImage(TSharedPtr<SWidget> ExpanderWidget, bool bIsUserCollections) const
+{
+	const bool bIsItemExpanded = bIsUserCollections ? bAreUserCollectionsExpanded : bAreProjectCollectionsExpanded;
+
+	FName ResourceName;
+	if (bIsItemExpanded)
+	{
+		if (ExpanderWidget->IsHovered())
+		{
+			static FName ExpandedHoveredName = "TreeArrow_Expanded_Hovered";
+			ResourceName = ExpandedHoveredName;
+		}
+		else
+		{
+			static FName ExpandedName = "TreeArrow_Expanded";
+			ResourceName = ExpandedName;
+		}
+	}
+	else
+	{
+		if (ExpanderWidget->IsHovered())
+		{
+			static FName CollapsedHoveredName = "TreeArrow_Collapsed_Hovered";
+			ResourceName = CollapsedHoveredName;
+		}
+		else
+		{
+			static FName CollapsedName = "TreeArrow_Collapsed";
+			ResourceName = CollapsedName;
+		}
+	}
+
+	return FCoreStyle::Get().GetBrush(ResourceName);
+}
+
+#undef LOCTEXT_NAMESPACE
\ No newline at end of file
diff --git a/Engine/Plugins/Experimental/ToolPresets/Source/PresetEditor/Private/Widgets/SPresetManager.h b/Engine/Plugins/Experimental/ToolPresets/Source/PresetEditor/Private/Widgets/SPresetManager.h
new file mode 100644
--- /dev/null
+++ b/Engine/Plugins/Experimental/ToolPresets/Source/PresetEditor/Private/Widgets/SPresetManager.h
@@ -0,0 +1,220 @@
+// Copyright Epic Games, Inc. All Rights Reserved.
+
+#pragma once
+
+#include "CoreMinimal.h"
+#include "Input/Reply.h"
+#include "Widgets/DeclarativeSyntaxSupport.h"
+#include "Widgets/SCompoundWidget.h"
+#include "Widgets/Views/STableViewBase.h"
+#include "Widgets/Views/STableRow.h"
+#include "Misc/ITransaction.h"
+
+class FTransaction;
+class UPresetUserSettings;
+class SPositiveActionButton;
+class SNegativeActionButton;
+class SSplitter;
+class SVerticalBox;
+class SEditableTextBox;
+template<class ItemType> class STreeView;
+template<class ItemType> class SListView;
+
+
+
+/**
+ * Implements the preset manager panel.
+ */
+class SPresetManager
+	: public SCompoundWidget
+{
+
+	struct FPresetCollectionInfo
+	{
+		FSoftObjectPath PresetCollectionPath;
+		bool bCollectionEnabled;
+
+		FPresetCollectionInfo(FSoftObjectPath InPresetCollectionPath, bool bInEnabled)
+			: PresetCollectionPath(InPresetCollectionPath), bCollectionEnabled(bInEnabled)
+		{}
+	};
+
+public:
+
+	struct FPresetViewEntry
+	{
+		enum class EEntryType : uint8
+		{
+			Collection,
+			Tool,
+			Preset
+		};
+		
+		EEntryType EntryType;
+
+		// Used for Collections/Tool Entries
+		bool bEnabled = false;
+		FSoftObjectPath CollectionPath;
+		FText EntryLabel;
+		FSlateBrush EntryIcon;
+		int32 Count = 0;
+
+		// Used for Preset/Tool entries
+		FString ToolName = "";
+		int32 PresetIndex = 0;
+		FString PresetLabel = "";
+		FString PresetTooltip = "";
+
+		TSharedPtr< FPresetViewEntry> Parent;
+		TArray<TSharedPtr< FPresetViewEntry> > Children;
+
+		// Collection Constructor
+		FPresetViewEntry(bool bEnabledIn, FSoftObjectPath CollectionPathIn, FText EntryLabelIn, int32 CountIn)
+			: EntryType(EEntryType::Collection),
+			  bEnabled(bEnabledIn), 
+			  CollectionPath(CollectionPathIn),
+			  EntryLabel(EntryLabelIn),
+			  Count(CountIn)
+		{}
+
+		// Tool Constructor
+		FPresetViewEntry(FText EntryLabelIn, FSlateBrush EntryIconIn, FSoftObjectPath CollectionPathIn, FString ToolNameIn, int32 CountIn)
+		  :	EntryType(EEntryType::Tool),
+			CollectionPath(CollectionPathIn),
+			EntryLabel(EntryLabelIn),
+			EntryIcon(EntryIconIn),
+			Count(CountIn),
+			ToolName(ToolNameIn)
+		{}
+
+		// Preset Constructor
+		FPresetViewEntry(FString ToolNameIn, int32 PresetIndexIn, FString PresetLabelIn, FString PresetTooltipIn, FText EntryLabelIn)
+			: EntryType(EEntryType::Preset),
+			  EntryLabel(EntryLabelIn),
+			  ToolName(ToolNameIn),
+			  PresetIndex(PresetIndexIn),
+			  PresetLabel(PresetLabelIn),
+			  PresetTooltip(PresetTooltipIn)
+		{}
+
+		bool HasSameMetadata(FPresetViewEntry& Other)
+		{
+			bool bIsEqual = 
+				 EntryType == Other.EntryType &&
+				 CollectionPath == Other.CollectionPath &&
+				 Count == Other.Count &&
+				 ToolName.Equals(Other.ToolName) &&
+				 PresetIndex == Other.PresetIndex &&
+				 PresetLabel.Equals(Other.PresetLabel) &&
+				 PresetTooltip.Equals(Other.PresetTooltip) &&
+				 Children.Num() == Other.Children.Num();
+
+			for (int32 ChildIndex = 0; ChildIndex < Children.Num(); ++ChildIndex)
+			{
+				bIsEqual = bIsEqual && (*Children[ChildIndex]).HasSameMetadata(*Other.Children[ChildIndex]);
+			}
+			return bIsEqual;
+		}
+
+		bool operator==(FPresetViewEntry& Other)
+		{
+			bool bIsEqual = bEnabled == Other.bEnabled &&
+				 EntryType == Other.EntryType &&
+				 CollectionPath == Other.CollectionPath &&
+				 Count == Other.Count &&
+				 EntryLabel.EqualTo(Other.EntryLabel) &&
+				 ToolName.Equals(Other.ToolName) &&
+				 PresetIndex == Other.PresetIndex &&
+				 PresetLabel.Equals(Other.PresetLabel) &&
+				 PresetTooltip.Equals(Other.PresetTooltip) &&
+				 Children.Num() == Other.Children.Num();
+
+			for (int32 ChildIndex = 0; ChildIndex < Children.Num(); ++ChildIndex)
+			{
+				bIsEqual = bIsEqual && (*Children[ChildIndex]) == (*Other.Children[ChildIndex]);
+			}
+
+			return bIsEqual;
+		}
+	};
+
+
+
+	SLATE_BEGIN_ARGS(SPresetManager) { }
+	SLATE_END_ARGS()
+
+	virtual ~SPresetManager();
+
+	/**
+	 * Construct this widget
+	 *
+	 * @param InArgs The declaration data for this widget.
+	 */
+	void Construct( const FArguments& InArgs );
+	
+	//~ Begin SWidget Interface
+	virtual void Tick( const FGeometry& AllottedGeometry, const double InCurrentTime, const float InDeltaTime ) override;
+	//~ End SWidget Interface
+
+protected:
+
+private:
+
+	int32 GetTotalPresetCount() const;
+
+	TSharedRef<ITableRow> HandleTreeGenerateRow(TSharedPtr<FPresetViewEntry> TreeEntry, const TSharedRef<STableViewBase>& OwnerTable);
+	void HandleTreeGetChildren(TSharedPtr<FPresetViewEntry> TreeEntry, TArray< TSharedPtr<FPresetViewEntry> >& ChildrenOut);
+	void HandleTreeSelectionChanged(TSharedPtr<FPresetViewEntry> TreeEntry, ESelectInfo::Type);
+	void HandleUserTreeSelectionChanged(TSharedPtr<FPresetViewEntry> TreeEntry, ESelectInfo::Type);
+	EVisibility ProjectPresetCollectionsVisibility() const;
+
+	void GeneratePresetList(TSharedPtr<FPresetViewEntry> TreeEntry);
+	TSharedRef<ITableRow> HandleListGenerateRow(TSharedPtr<FPresetViewEntry> TreeEntry, const TSharedRef<STableViewBase>& OwnerTable);
+    void HandleListSelectionChanged(TSharedPtr<FPresetViewEntry> TreeEntry, ESelectInfo::Type SelectInfo);
+
+	EVisibility EditAreaVisibility() const;
+	void SetCollectionEnabled(TSharedPtr<FPresetViewEntry> TreeEntry, ECheckBoxState State);
+	void DeletePresetFromCollection(TSharedPtr< FPresetViewEntry > Entry);
+
+	void DeleteSelectedUserPresetCollection();
+	void AddNewUserPresetCollection();
+
+	void SetPresetLabel(TSharedPtr< FPresetViewEntry >, FText InLabel);
+	void SetPresetTooltip(TSharedPtr< FPresetViewEntry >, FText InTooltip);
+
+	const FSlateBrush* GetProjectCollectionsExpanderImage() const;
+	const FSlateBrush* GetUserCollectionsExpanderImage() const;
+	const FSlateBrush* GetExpanderImage(TSharedPtr<SWidget> ExpanderWidget, bool bIsUserCollections) const;
+
+private:
+
+	TWeakObjectPtr<UPresetUserSettings> UserSettings;
+
+	bool bAreProjectCollectionsExpanded = true;
+	TSharedPtr<SButton> ProjectCollectionsExpander;
+	TArray< TSharedPtr< FPresetViewEntry > > ProjectCollectionsDataList;
+	TSharedPtr<STreeView<TSharedPtr<FPresetViewEntry> > > ProjectPresetCollectionTreeView;
+
+	bool bAreUserCollectionsExpanded = true;
+	TSharedPtr<SButton> UserCollectionsExpander;
+	TArray< TSharedPtr< FPresetViewEntry > > UserCollectionsDataList;
+	TSharedPtr<STreeView<TSharedPtr<FPresetViewEntry> > > UserPresetCollectionTreeView;
+
+	TArray< TSharedPtr< FPresetViewEntry > > PresetDataList;
+	TSharedPtr<SListView<TSharedPtr<FPresetViewEntry> > > PresetListView;
+
+	int32 TotalPresetCount;
+	bool bHasActiveCollection;
+	FText ActiveCollectionLabel;
+	bool bIsActiveCollectionEnabled;
+
+	TSharedPtr<SSplitter> Splitter;
+
+	TSharedPtr<SVerticalBox>  EditPresetArea;
+	TSharedPtr<SEditableTextBox> EditPresetLabel;
+	TSharedPtr<SEditableTextBox> EditPresetTooltip;
+	TSharedPtr<FPresetViewEntry> ActivePresetToEdit;
+
+	TSharedPtr<SPositiveActionButton> AddUserPresetButton;
+	TSharedPtr<SNegativeActionButton> DeleteUserPresetButton;
+};
diff --git a/Engine/Plugins/Experimental/ToolPresets/Source/PresetEditor/Public/IPresetEditorModule.h b/Engine/Plugins/Experimental/ToolPresets/Source/PresetEditor/Public/IPresetEditorModule.h
new file mode 100644
--- /dev/null
+++ b/Engine/Plugins/Experimental/ToolPresets/Source/PresetEditor/Public/IPresetEditorModule.h
@@ -0,0 +1,21 @@
+// Copyright Epic Games, Inc. All Rights Reserved.
+
+#pragma once
+
+#include "CoreMinimal.h"
+#include "Modules/ModuleInterface.h"
+#include "Modules/ModuleManager.h"
+
+PRESETEDITOR_API extern const FName PresetEditorTabName;
+
+class PRESETEDITOR_API IPresetEditorModule : public IModuleInterface
+{
+public:
+
+	static IPresetEditorModule& Get()
+	{
+		return FModuleManager::LoadModuleChecked<IPresetEditorModule>("PresetEditor");
+	}
+
+	virtual void ExecuteOpenPresetEditor() = 0;
+};
diff --git a/Engine/Plugins/Experimental/ToolPresets/Source/PresetEditor/Public/PresetSettings.h b/Engine/Plugins/Experimental/ToolPresets/Source/PresetEditor/Public/PresetSettings.h
new file mode 100644
--- /dev/null
+++ b/Engine/Plugins/Experimental/ToolPresets/Source/PresetEditor/Public/PresetSettings.h
@@ -0,0 +1,66 @@
+// Copyright Epic Games, Inc. All Rights Reserved.
+
+#pragma once
+#include "CoreMinimal.h"
+#include "UObject/ObjectMacros.h"
+#include "UObject/Object.h"
+#include "EditorConfigBase.h"
+#include "Engine/DeveloperSettings.h"
+
+#include "PresetSettings.generated.h"
+
+/**
+ * Implements the settings for the PresetEditor.
+ */
+UCLASS(EditorConfig = "UPresetUserSettings")
+class PRESETEDITOR_API UPresetUserSettings : public UEditorConfigBase
+{
+public:
+	GENERATED_BODY()
+
+	UPROPERTY(EditAnywhere, Category = "Interactive Tool Presets|Collections", meta = (EditorConfig, AllowedClasses = "/Script/PresetAsset.InteractiveToolsPresetCollectionAsset"))
+	TSet<FSoftObjectPath> EnabledPresetCollections;// = { FSoftObjectPath(FTopLevelAssetPath("/ToolPresets/Presets/_DefaultCollection", "_DefaultCollection")) };
+
+	// TODO: Ideally the above property would be able to store if the default collection was enabled or not.
+	// However, currently the EditorConfig system has some trouble with defaults and array/set properties.
+	// To overcome this, we store it separately which doesn't seem to break as badly.
+
+	UPROPERTY(VisibleAnywhere, Category = "Interactive Tool Presets|Collections", meta = (EditorConfig, AllowedClasses = "/Script/PresetAsset.InteractiveToolsPresetCollectionAsset"))
+	FSoftObjectPath DefaultCollectionPath = FSoftObjectPath(FTopLevelAssetPath("/ToolPresets/Presets/_DefaultCollection", "_DefaultCollection"));
+
+	UPROPERTY(EditAnywhere, Category = "Interactive Tool Presets|Collections", meta = (EditorConfig))
+	bool bDefaultCollectionEnabled = true;
+
+	static void Initialize();
+	static UPresetUserSettings* Get();
+
+private:
+	static TObjectPtr<UPresetUserSettings> Instance;
+};
+
+
+/**
+ * Implements the settings for the PresetEditor.
+ */
+UCLASS(config = Editor)
+class PRESETEDITOR_API UPresetProjectSettings
+	: public UDeveloperSettings
+{
+public:
+
+	// UDeveloperSettings overrides
+
+	virtual FName GetContainerName() const override { return FName("Project"); }
+	virtual FName GetCategoryName() const override { return FName("Plugins"); }
+	virtual FName GetSectionName() const override { return FName("Interactive Tool Presets"); }
+
+	virtual FText GetSectionText() const override { return NSLOCTEXT("PresetSettings", "SectionText", "Interactive Tool Presets"); };
+	virtual FText GetSectionDescription() const override { return NSLOCTEXT("PresetSettings", "SectionDescription", "Manage preset settings at the project level."); };
+
+public:
+	GENERATED_BODY()
+
+	/* Controls which preset collection assets are to be loaded for this project.  */
+	UPROPERTY(config, EditAnywhere, Category = "Interactive Tool Presets|Collections", meta = (AllowedClasses = "/Script/PresetAsset.InteractiveToolsPresetCollectionAsset"))
+	TSet<FSoftObjectPath> LoadedPresetCollections;
+};
\ No newline at end of file
diff --git a/Engine/Plugins/Experimental/ToolPresets/ToolPresets.uplugin b/Engine/Plugins/Experimental/ToolPresets/ToolPresets.uplugin
--- a/Engine/Plugins/Experimental/ToolPresets/ToolPresets.uplugin
+++ b/Engine/Plugins/Experimental/ToolPresets/ToolPresets.uplugin
@@ -11,7 +11,7 @@
 	"MarketplaceURL": "",
 	"SupportURL": "",
 	"EnabledByDefault": true,
-	"CanContainContent": false,
+	"CanContainContent": true,
 	"IsBetaVersion": false,
 	"Installed": false,
 	"IsExperimentalVersion": true,
@@ -20,6 +20,12 @@
 			"Name": "PresetAsset",
 			"Type": "Editor",
 			"LoadingPhase": "Default"
+		},
+		{
+			"Name": "PresetEditor",
+			"Type": "Editor",
+			"LoadingPhase": "Default"
 		}
+
 	]
 }
diff --git a/Engine/Source/Developer/WidgetRegistration/Private/ToolkitBuilder.cpp b/Engine/Source/Developer/WidgetRegistration/Private/ToolkitBuilder.cpp
--- a/Engine/Source/Developer/WidgetRegistration/Private/ToolkitBuilder.cpp
+++ b/Engine/Source/Developer/WidgetRegistration/Private/ToolkitBuilder.cpp
@@ -479,6 +479,17 @@ void FToolkitBuilder::DefineWidget()
 		]
 	];
 	
+	if (ToolkitSections->ToolPresetArea)
+	{
+		ToolkitWidgetVBox->AddSlot()
+			.AutoHeight()
+			.HAlign(HAlign_Fill)
+			.Padding(5)
+			[
+				ToolkitSections->ToolPresetArea->AsShared()
+			];
+	}
+
 	if (ToolkitSections->ToolWarningArea)
 	{
 		ToolkitWidgetVBox->AddSlot()
diff --git a/Engine/Source/Developer/WidgetRegistration/Public/ToolkitBuilder.h b/Engine/Source/Developer/WidgetRegistration/Public/ToolkitBuilder.h
--- a/Engine/Source/Developer/WidgetRegistration/Public/ToolkitBuilder.h
+++ b/Engine/Source/Developer/WidgetRegistration/Public/ToolkitBuilder.h
@@ -21,6 +21,7 @@ struct WIDGETREGISTRATION_API FToolkitSections
 {
 	TSharedPtr<STextBlock> ModeWarningArea = nullptr;
 	TSharedPtr<STextBlock> ToolWarningArea = nullptr;
+	TSharedPtr<SWidget> ToolPresetArea = nullptr;
 	TSharedPtr<IDetailsView> DetailsView = nullptr;
 	TSharedPtr<SWidget> Footer = nullptr;
 };
