ToolPresets: Fix copyright boilerplate issue.

#rnx

diff --git a/Engine/Plugins/Experimental/ToolPresets/Source/PresetEditor/Private/PresetSettings.cpp b/Engine/Plugins/Experimental/ToolPresets/Source/PresetEditor/Private/PresetSettings.cpp
--- a/Engine/Plugins/Experimental/ToolPresets/Source/PresetEditor/Private/PresetSettings.cpp
+++ b/Engine/Plugins/Experimental/ToolPresets/Source/PresetEditor/Private/PresetSettings.cpp
@@ -1,4 +1,4 @@
-/// Copyright Epic Games, Inc. All Rights Reserved.
+// Copyright Epic Games, Inc. All Rights Reserved.
 
 #include "PresetSettings.h"
 
