ToolPresets: Address issues around CI cooking steps with default preset creation

Reimplements default presets to be stored on disk as JSON configuration files, loaded at editor startup into a transient collection asset.

diff --git a/Engine/Plugins/Editor/ModelingToolsEditorMode/Source/ModelingToolsEditorMode/Private/ModelingToolsEditorModeToolkit.cpp b/Engine/Plugins/Editor/ModelingToolsEditorMode/Source/ModelingToolsEditorMode/Private/ModelingToolsEditorModeToolkit.cpp
--- a/Engine/Plugins/Editor/ModelingToolsEditorMode/Source/ModelingToolsEditorMode/Private/ModelingToolsEditorModeToolkit.cpp
+++ b/Engine/Plugins/Editor/ModelingToolsEditorMode/Source/ModelingToolsEditorMode/Private/ModelingToolsEditorModeToolkit.cpp
@@ -58,6 +58,7 @@
 #include "AssetRegistry/AssetRegistryModule.h"
 #include "ModelingWidgets/SToolInputAssetComboPanel.h"
 #include "Fonts/SlateFontInfo.h"
+#include "PresetAssetSubsystem.h"
 
 
 #define LOCTEXT_NAMESPACE "FModelingToolsEditorModeToolkit"
@@ -80,12 +81,42 @@ static TAutoConsoleVariable<int32> CVarEnableToolPresets(
 namespace FModelingToolsEditorModeToolkitLocals
 {
 	typedef TFunction<void(UInteractiveToolsPresetCollectionAsset& Preset, UInteractiveTool& Tool)> PresetAndToolFunc;
+	typedef TFunction<void(UInteractiveToolsPresetCollectionAsset& Preset)> PresetOnlyFunc;
+
+	void ExecuteWithPreset(const FSoftObjectPath& PresetPath, PresetOnlyFunc Function)
+	{
+		UInteractiveToolsPresetCollectionAsset* Preset = nullptr;
+
+		UPresetAssetSubsystem* PresetAssetSubsystem = GEditor->GetEditorSubsystem<UPresetAssetSubsystem>();
+
+		if (PresetPath.IsNull() && ensure(PresetAssetSubsystem))
+		{
+			Preset = PresetAssetSubsystem->GetDefaultCollection();
+		}
+		if (PresetPath.IsAsset())
+		{
+			Preset = Cast<UInteractiveToolsPresetCollectionAsset>(PresetPath.TryLoad());
+		}
+		if (!Preset)
+		{
+			return;
+		}
+		Function(*Preset);
+
+
+	}
 
 	void ExecuteWithPresetAndTool(UEdMode& EdMode, EToolSide ToolSide, const FSoftObjectPath& PresetPath, PresetAndToolFunc Function)
 	{
 		UInteractiveToolsPresetCollectionAsset* Preset = nullptr;		
 		UInteractiveTool* Tool = EdMode.GetToolManager()->GetActiveTool(EToolSide::Left);
 
+		UPresetAssetSubsystem* PresetAssetSubsystem = GEditor->GetEditorSubsystem<UPresetAssetSubsystem>();
+
+		if (PresetPath.IsNull() && ensure(PresetAssetSubsystem))
+		{
+			Preset = PresetAssetSubsystem->GetDefaultCollection();
+		}
 		if (PresetPath.IsAsset())
 		{
 			Preset = Cast<UInteractiveToolsPresetCollectionAsset>(PresetPath.TryLoad());
@@ -109,6 +140,19 @@ namespace FModelingToolsEditorModeToolkitLocals
 	}
 }
 
+class FRecentPresetCollectionProvider : public SToolInputAssetComboPanel::IRecentAssetsProvider
+{
+	public:
+		//~ SToolInputAssetComboPanel::IRecentAssetsProvider interface		
+		virtual TArray<FAssetData> GetRecentAssetsList() override { return RecentPresetCollectionList; }		
+		virtual void NotifyNewAsset(const FAssetData& NewAsset) {
+			RecentPresetCollectionList.AddUnique(NewAsset);
+		};
+
+	protected:
+		TArray<FAssetData> RecentPresetCollectionList;
+};
+
 FModelingToolsEditorModeToolkit::FModelingToolsEditorModeToolkit()
 {
 	UModelingModeEditableToolPaletteConfig::Initialize();
@@ -116,6 +160,8 @@ FModelingToolsEditorModeToolkit::FModelingToolsEditorModeToolkit()
 
 	UPresetUserSettings::Initialize();
 	UPresetUserSettings::Get()->LoadEditorConfig();
+
+	RecentPresetCollectionProvider = MakeShared< FRecentPresetCollectionProvider>();
 }
 
 FModelingToolsEditorModeToolkit::~FModelingToolsEditorModeToolkit()
@@ -138,6 +184,8 @@ FModelingToolsEditorModeToolkit::~FModelingToolsEditorModeToolkit()
 	Settings->OnModified.Remove(AssetSettingsModifiedHandle);
 	GetScriptableEditorMode()->GetInteractiveToolsContext(EToolsContextScope::EdMode)->OnToolNotificationMessage.RemoveAll(this);
 	GetScriptableEditorMode()->GetInteractiveToolsContext(EToolsContextScope::EdMode)->OnToolWarningMessage.RemoveAll(this);
+
+	RecentPresetCollectionProvider = nullptr;
 }
 
 
@@ -540,7 +588,7 @@ void FModelingToolsEditorModeToolkit::Init(const TSharedPtr<IToolkitHost>& InitT
 		GetToolkitHost()->AddViewportOverlayWidget(SelectionPaletteOverlayWidget.ToSharedRef());
 	}
 
-	CurrentPreset = FSoftObjectPath(FTopLevelAssetPath("/ToolPresets/Presets/_DefaultCollection", "_DefaultCollection"));
+	CurrentPreset = FSoftObjectPath(); // Default to the default collection by leaving this null.
 }
 
 void FModelingToolsEditorModeToolkit::MakeToolShutdownOverlayWidget()
@@ -954,48 +1002,45 @@ TSharedRef<SWidget> FModelingToolsEditorModeToolkit::GetPresetCreateButtonConten
 			.AutoHeight()
 			.Padding(5.0f)
 			[
-				SNew(SExpandableArea)
-				.AreaTitle(FText::FromString("Advanced"))
-			    .InitiallyCollapsed(true)
-				.BodyContent()
+				SNew(SHorizontalBox)
+				+ SHorizontalBox::Slot()
+				.AutoWidth()
 				[
-
-					SNew(SHorizontalBox)
-					+ SHorizontalBox::Slot()
-					.AutoWidth()
+					SNew(SToolInputAssetComboPanel)
+					.AssetClassType(UInteractiveToolsPresetCollectionAsset::StaticClass())
+					.OnSelectionChanged(this, &FModelingToolsEditorModeToolkit::HandlePresetAssetChanged)
+					.ToolTipText(LOCTEXT("ToolPresets_CreatePresetCollection_Tooltip", "The asset in which to store this new preset."))
+					//.RecentAssetsProvider(RecentPresetCollectionProvider) // TODO: Improve this widget before enabling this feature
+					.FlyoutTileSize(FVector2D(80, 80))
+					.ComboButtonTileSize(FVector2D(80, 80))
+				]
+				+ SHorizontalBox::Slot()
+				[
+					SNew(SVerticalBox)
+					+ SVerticalBox::Slot()
+					.AutoHeight()
+					.Padding(10, 5)
 					[
-						SNew(SToolInputAssetComboPanel)
-						.AssetClassType(UInteractiveToolsPresetCollectionAsset::StaticClass())
-						.OnSelectionChanged(this, &FModelingToolsEditorModeToolkit::HandlePresetAssetChanged)
-						.ToolTipText(LOCTEXT("ToolPresets_CreatePresetCollection_Tooltip", "The asset in which to store this new preset."))
-						.FlyoutTileSize(FVector2D(80, 80))
-						.ComboButtonTileSize(FVector2D(80, 80))
-						.InitiallySelectedAsset( FAssetData(FName("/ToolPresets/Presets/_DefaultCollection"),
-															FName("/ToolPresets/Presets"),
-															FName("_DefaultCollection"),
-															UInteractiveToolsPresetCollectionAsset::StaticClass()->GetClassPathName() ))
+						SNew(STextBlock)
+						.Text(LOCTEXT("ToolPresets_CreatePresetCollection", "Collection"))
+						.Font(FSlateFontInfo(FCoreStyle::GetDefaultFont(),12, "Bold"))
 					]
-					+ SHorizontalBox::Slot()
+					+ SVerticalBox::Slot()
+					.AutoHeight()
+					.Padding(10, 5)
 					[
-						SNew(SVerticalBox)
-						+ SVerticalBox::Slot()
-						.AutoHeight()
-						.Padding(10, 5)
-						[
-							SNew(STextBlock)
-							.Text(LOCTEXT("ToolPresets_CreatePresetCollection", "Collection"))
-							.Font(FSlateFontInfo(FCoreStyle::GetDefaultFont(),12, "Bold"))
-						]
-						+ SVerticalBox::Slot()
-						.AutoHeight()
-						.Padding(10, 5)
-						[
-							SNew(STextBlock)
-							.Text_Lambda([this]() { return CurrentPresetLabel; })
-						]
-
+						SNew(STextBlock)
+						.Text_Lambda([this](){
+						if (CurrentPresetLabel.IsEmpty())
+						{
+							return LOCTEXT("NewPresetNoCollectionSpecifiedMessage", "None - Preset will be added to Editor Presets.");
+						}
+						else {
+							return CurrentPresetLabel;
+						}})
 					]
-				]		
+
+				]					
 			],
 			FolderDialogArguments, true);
 
@@ -1017,15 +1062,10 @@ TSharedRef<SWidget> FModelingToolsEditorModeToolkit::GetPresetCreateButtonConten
 		FPresetsByNameMap PresetsByCollectionName;
 		for (TSharedPtr<FToolPresetOption> ToolPresetOption : AvailablePresetsForTool)
 		{
-			UInteractiveToolsPresetCollectionAsset* Preset = nullptr;
-			if (ToolPresetOption->PresetCollection.IsAsset())
-			{
-				Preset = Cast<UInteractiveToolsPresetCollectionAsset>(ToolPresetOption->PresetCollection.TryLoad());
-			}
-			if (Preset)
-			{
-				PresetsByCollectionName.FindOrAdd(Preset->CollectionLabel.ToString()).Add(ToolPresetOption);
-			}
+			FModelingToolsEditorModeToolkitLocals::ExecuteWithPreset(ToolPresetOption->PresetCollection,
+				[this, &PresetsByCollectionName, &ToolPresetOption](UInteractiveToolsPresetCollectionAsset& Preset) {
+					PresetsByCollectionName.FindOrAdd(Preset.CollectionLabel.ToString()).Add(ToolPresetOption);
+				});
 		}
 
 		for (FPresetsByNameMap::TConstIterator Iterator = PresetsByCollectionName.CreateConstIterator(); Iterator; ++Iterator)
@@ -1151,7 +1191,7 @@ void FModelingToolsEditorModeToolkit::RebuildPresetListForTool(bool bSettingsOpe
 	TSet<FSoftObjectPath> AvailablePresetCollections = UserSettings->EnabledPresetCollections.Intersect( ProjectSettings->LoadedPresetCollections.Union(AllUserPresets));
 	if (UserSettings->bDefaultCollectionEnabled)
 	{
-		AvailablePresetCollections.Add(UserSettings->DefaultCollectionPath);
+		AvailablePresetCollections.Add(FSoftObjectPath());
 	}
 
 
@@ -1257,6 +1297,12 @@ void FModelingToolsEditorModeToolkit::CreateNewPresetInCollection(const FString&
 		AvailablePresetsForTool.Add(NewOption);
 		
 	});
+
+	UPresetAssetSubsystem* PresetAssetSubsystem = GEditor->GetEditorSubsystem<UPresetAssetSubsystem>();
+	if (CollectionPath.IsNull() && ensure(PresetAssetSubsystem))
+	{
+		ensure(PresetAssetSubsystem->SaveDefaultCollection());
+	}
 }
 
 void FModelingToolsEditorModeToolkit::UpdatePresetInCollection(const FToolPresetOption& PresetToEditIn, bool bUpdateStoredPresetValues)
@@ -1285,6 +1331,12 @@ void FModelingToolsEditorModeToolkit::UpdatePresetInCollection(const FToolPreset
 		});
 
 	RebuildPresetListForTool(false);
+
+	UPresetAssetSubsystem* PresetAssetSubsystem = GEditor->GetEditorSubsystem<UPresetAssetSubsystem>();
+	if (PresetToEditIn.PresetCollection.IsNull() && ensure(PresetAssetSubsystem))
+	{
+		ensure(PresetAssetSubsystem->SaveDefaultCollection());
+	}
 }
 
 
diff --git a/Engine/Plugins/Editor/ModelingToolsEditorMode/Source/ModelingToolsEditorMode/Public/ModelingToolsEditorModeToolkit.h b/Engine/Plugins/Editor/ModelingToolsEditorMode/Source/ModelingToolsEditorMode/Public/ModelingToolsEditorModeToolkit.h
--- a/Engine/Plugins/Editor/ModelingToolsEditorMode/Source/ModelingToolsEditorMode/Public/ModelingToolsEditorModeToolkit.h
+++ b/Engine/Plugins/Editor/ModelingToolsEditorMode/Source/ModelingToolsEditorMode/Public/ModelingToolsEditorModeToolkit.h
@@ -18,6 +18,7 @@ class SButton;
 class STextBlock;
 class UGeometrySelectionManager;
 class UInteractiveToolsPresetCollectionAsset;
+class FRecentPresetCollectionProvider;
 
 struct FToolPresetOption
 {
@@ -159,6 +160,7 @@ private:
 
 	TArray<TSharedPtr<FToolPresetOption>> AvailablePresetsForTool;
 	TSharedPtr<SEditableComboBox<TSharedPtr<FString>>> PresetComboBox;
+	TSharedPtr<FRecentPresetCollectionProvider> RecentPresetCollectionProvider;
 
 	TSharedRef<SWidget> GetPresetCreateButtonContent();
 
diff --git a/Engine/Plugins/Experimental/ToolPresets/Source/PresetAsset/PresetAsset.Build.cs b/Engine/Plugins/Experimental/ToolPresets/Source/PresetAsset/PresetAsset.Build.cs
--- a/Engine/Plugins/Experimental/ToolPresets/Source/PresetAsset/PresetAsset.Build.cs
+++ b/Engine/Plugins/Experimental/ToolPresets/Source/PresetAsset/PresetAsset.Build.cs
@@ -25,7 +25,7 @@ public class PresetAsset: ModuleRules
 		PublicDependencyModuleNames.AddRange(
 			new string[]
 			{
-				"Core",
+				"Core",				
 			}
 			);
 			
@@ -41,6 +41,8 @@ public class PresetAsset: ModuleRules
 				"AssetDefinition",
 				"JsonUtilities",
 				"Json",
+				"EditorSubsystem",
+				"EditorConfig",
 			}
 			);
 		
diff --git a/Engine/Plugins/Experimental/ToolPresets/Source/PresetAsset/Private/PresetAssetSubsystem.cpp b/Engine/Plugins/Experimental/ToolPresets/Source/PresetAsset/Private/PresetAssetSubsystem.cpp
new file mode 100644
--- /dev/null
+++ b/Engine/Plugins/Experimental/ToolPresets/Source/PresetAsset/Private/PresetAssetSubsystem.cpp
@@ -0,0 +1,56 @@
+// Copyright Epic Games, Inc. All Rights Reserved.
+
+#include "PresetAssetSubsystem.h"
+#include "PresetAsset.h"
+#include "EditorConfigSubsystem.h"
+
+#define LOCTEXT_NAMESPACE "PresetAssetSubsystem"
+
+void UPresetAssetSubsystem::Initialize(FSubsystemCollectionBase& Collection)
+{
+	Collection.InitializeDependency(UEditorConfigSubsystem::StaticClass());
+
+	InitializeDefaultCollection();
+}
+
+void UPresetAssetSubsystem::Deinitialize()
+{
+	if (DefaultCollection)
+	{
+		DefaultCollection->SaveEditorConfig();
+		DefaultCollection = nullptr;
+	}	
+}
+
+UInteractiveToolsPresetCollectionAsset* UPresetAssetSubsystem::GetDefaultCollection()
+{
+	return DefaultCollection;
+}
+
+bool UPresetAssetSubsystem::SaveDefaultCollection()
+{
+	if (DefaultCollection)
+	{
+		return DefaultCollection->SaveEditorConfig();
+	}
+	return false;
+}
+
+
+void UPresetAssetSubsystem::InitializeDefaultCollection()
+{
+	/*
+	* We're storing the default collection as a JSON file instead of an asset
+	* on disk for a few reasons. First it avoids issues around automatically
+	* creating assets, both from avoiding build system issues and from a more
+	* philisophical point about requiring user involvement. Second, it helps
+	* compartmentalizing the "default" collection as more of Editor preferences,
+	* rather than a specific collection which has purpose and can be shared around. 
+	*/
+
+	DefaultCollection = NewObject<UInteractiveToolsPresetCollectionAsset>();
+	DefaultCollection->CollectionLabel = LOCTEXT("DefaultCollectionLabel", "Editor Presets");
+	DefaultCollection->LoadEditorConfig();
+}
+
+#undef LOCTEXT_NAMESPACE
\ No newline at end of file
diff --git a/Engine/Plugins/Experimental/ToolPresets/Source/PresetAsset/Public/PresetAsset.h b/Engine/Plugins/Experimental/ToolPresets/Source/PresetAsset/Public/PresetAsset.h
--- a/Engine/Plugins/Experimental/ToolPresets/Source/PresetAsset/Public/PresetAsset.h
+++ b/Engine/Plugins/Experimental/ToolPresets/Source/PresetAsset/Public/PresetAsset.h
@@ -7,6 +7,7 @@
 #include "UObject/ObjectMacros.h"
 #include "Factories/Factory.h"
 #include "AssetToolsModule.h"
+#include "EditorConfigBase.h"
 
 #include "PresetAsset.generated.h"
 
@@ -17,13 +18,13 @@ struct PRESETASSET_API FInteractiveToolPresetDefintion
 {
 	GENERATED_BODY()
 public:
-	UPROPERTY(VisibleAnywhere, Category = "PresetAsset")
+	UPROPERTY(VisibleAnywhere, Category = "PresetAsset", meta = (EditorConfig))
 	FString StoredProperties;
 
-	UPROPERTY(EditAnywhere, Category = "PresetAsset")
+	UPROPERTY(EditAnywhere, Category = "PresetAsset", meta = (EditorConfig))
 	FString Label;
 
-	UPROPERTY(EditAnywhere, Category = "PresetAsset")
+	UPROPERTY(EditAnywhere, Category = "PresetAsset", meta = (EditorConfig))
 	FString Tooltip;
 
 	bool IsValid() const;
@@ -36,22 +37,22 @@ struct PRESETASSET_API FInteractiveToolPresetStore
 {
 	GENERATED_BODY()
 public:
-	UPROPERTY(EditAnywhere, Category = "PresetAsset")
+	UPROPERTY(EditAnywhere, Category = "PresetAsset", meta = (EditorConfig))
 	TArray<FInteractiveToolPresetDefintion> NamedPresets;
 
-	UPROPERTY(VisibleAnywhere, Category = "PresetAsset")
+	UPROPERTY(VisibleAnywhere, Category = "PresetAsset", meta = (EditorConfig))
 	FText ToolLabel;
 
-	UPROPERTY(VisibleAnywhere, Category = "PresetAsset")
+	UPROPERTY(VisibleAnywhere, Category = "PresetAsset", meta = (EditorConfig))
 	FSlateBrush ToolIcon;
 };
 
 /**
  * Implements an asset that can be used to store tool settings as a named preset
  */
-UCLASS(BlueprintType, hidecategories=(Object))
-class PRESETASSET_API UInteractiveToolsPresetCollectionAsset
-	: public UObject
+UCLASS(BlueprintType, hidecategories=(Object), EditorConfig = "UInteractiveToolsPresetCollectionAsset_DefaultCollection")
+class PRESETASSET_API UInteractiveToolsPresetCollectionAsset 
+	: public UEditorConfigBase
 {
 	GENERATED_BODY()
 
@@ -64,10 +65,10 @@ public:
 	// Once we are satisifed with the data structure, planned accessors and mutators will include
 	// support for adding, removing, renaming, saving and retrieving presets.
 
-	UPROPERTY(VisibleAnywhere, Category = "PresetAsset")
+	UPROPERTY(VisibleAnywhere, Category = "PresetAsset", meta = (EditorConfig) )
 	TMap<FString, FInteractiveToolPresetStore > PerToolPresets;
 
-	UPROPERTY(EditAnywhere, Category = "PresetAsset")
+	UPROPERTY(EditAnywhere, Category = "PresetAsset", meta = (EditorConfig))
 	FText CollectionLabel;
 };
 
diff --git a/Engine/Plugins/Experimental/ToolPresets/Source/PresetAsset/Public/PresetAssetSubsystem.h b/Engine/Plugins/Experimental/ToolPresets/Source/PresetAsset/Public/PresetAssetSubsystem.h
new file mode 100644
--- /dev/null
+++ b/Engine/Plugins/Experimental/ToolPresets/Source/PresetAsset/Public/PresetAssetSubsystem.h
@@ -0,0 +1,33 @@
+// Copyright Epic Games, Inc. All Rights Reserved.
+
+#pragma once
+
+#include "CoreMinimal.h"
+#include "EditorSubsystem.h"
+#include "PresetAssetSubsystem.generated.h"
+
+class UInteractiveToolsPresetCollectionAsset;
+
+/**
+ * Using an editor subsystem allows us to make sure that we have a default preset asset whenever the editor exists
+ *  (and to avoid accidentally trying to make one when it doesn't, such as when running cooking scripts).
+ */
+UCLASS()
+class PRESETASSET_API UPresetAssetSubsystem : public UEditorSubsystem
+{
+	GENERATED_BODY()
+
+public:
+	virtual void Initialize(FSubsystemCollectionBase& Collection) override;
+	virtual void Deinitialize() override;
+
+	UInteractiveToolsPresetCollectionAsset* GetDefaultCollection();
+	bool SaveDefaultCollection();
+
+
+protected:
+	void InitializeDefaultCollection();
+
+	UPROPERTY()
+	TObjectPtr<UInteractiveToolsPresetCollectionAsset> DefaultCollection;
+};
\ No newline at end of file
diff --git a/Engine/Plugins/Experimental/ToolPresets/Source/PresetEditor/Private/PresetEditorModule.cpp b/Engine/Plugins/Experimental/ToolPresets/Source/PresetEditor/Private/PresetEditorModule.cpp
--- a/Engine/Plugins/Experimental/ToolPresets/Source/PresetEditor/Private/PresetEditorModule.cpp
+++ b/Engine/Plugins/Experimental/ToolPresets/Source/PresetEditor/Private/PresetEditorModule.cpp
@@ -59,8 +59,6 @@ public:
 	{
 		// Register slate style overrides
 		FPresetEditorStyle::Initialize();
-
-		GenerateDefaultCollection();
 	}
 
 	virtual bool SupportsDynamicReloading() override
@@ -86,45 +84,7 @@ private:
 		return DockTab;
 	}
 
-	void GenerateDefaultCollection()
-	{
-		// Load necessary modules
-		FAssetToolsModule& AssetToolsModule = FModuleManager::Get().LoadModuleChecked<FAssetToolsModule>("AssetTools");
-		FContentBrowserModule& ContentBrowserModule = FModuleManager::LoadModuleChecked<FContentBrowserModule>("ContentBrowser");
-		FAssetRegistryModule& AssetRegistryModule = FModuleManager::LoadModuleChecked<FAssetRegistryModule>(TEXT("AssetRegistry"));
-		IAssetRegistry& AssetRegistry = AssetRegistryModule.Get();
-
-		TArray<FAssetData> AssetData;
-
-		// Find (or create!) the desired package for this object
-		FString PackageName = "/ToolPresets/Presets/_DefaultCollection";
-		FString ObjectName = "_DefaultCollection";
-		const FString PackagePath = FPackageName::GetLongPackagePath(PackageName);
-
-		// Handle fully loading packages before creating new objects.
-		UPackage* Pkg = CreatePackage(*PackageName);
-		TArray<UPackage*> TopLevelPackages;
-		TopLevelPackages.Add(Pkg);
-		UPackageTools::HandleFullyLoadingPackages(TopLevelPackages, LOCTEXT("CreateANewObject", "Create a new object"));
-
-		// Check for an existing object
-		UObject* ExistingObject = StaticFindObject(UObject::StaticClass(), Pkg, *ObjectName);
-
-		if (!ExistingObject)
-		{
-			// Create object and package
-			UInteractiveToolsPresetCollectionAssetFactory* MyFactory = NewObject<UInteractiveToolsPresetCollectionAssetFactory>(UInteractiveToolsPresetCollectionAssetFactory::StaticClass()); // Can omit, and a default factory will be used
-			UObject* NewObject = AssetToolsModule.Get().CreateAsset(ObjectName, PackagePath, UInteractiveToolsPresetCollectionAsset::StaticClass(), MyFactory);
-			UInteractiveToolsPresetCollectionAsset* NewCollection = ExactCast<UInteractiveToolsPresetCollectionAsset>(NewObject);
-			NewCollection->CollectionLabel = FText::FromString("Default Collection");
-			FSavePackageArgs SavePackageArgs;
-			SavePackageArgs.TopLevelFlags = RF_Public | RF_Standalone;
-			UPackage::Save(Pkg, NewObject, *FPackageName::LongPackageNameToFilename(PackageName, FPackageName::GetAssetPackageExtension()), SavePackageArgs);
-
-			// Inform asset registry
-			AssetRegistry.AssetCreated(NewObject);
-		}
-	}
+
 };
 
 #undef LOCTEXT_NAMESPACE
diff --git a/Engine/Plugins/Experimental/ToolPresets/Source/PresetEditor/Private/Widgets/SPresetManager.cpp b/Engine/Plugins/Experimental/ToolPresets/Source/PresetEditor/Private/Widgets/SPresetManager.cpp
--- a/Engine/Plugins/Experimental/ToolPresets/Source/PresetEditor/Private/Widgets/SPresetManager.cpp
+++ b/Engine/Plugins/Experimental/ToolPresets/Source/PresetEditor/Private/Widgets/SPresetManager.cpp
@@ -31,6 +31,9 @@
 #include "SNegativeActionButton.h"
 #include "SPositiveActionButton.h"
 #include "PresetEditorStyle.h"
+#include "PresetAssetSubsystem.h"
+
+
 
 #define LOCTEXT_NAMESPACE "SPresetManager"
 
@@ -150,6 +153,16 @@ namespace PresetManagerLocals
 						[
 							SNew(STextBlock)
 							.Text(ViewEntry->EntryLabel)
+							.Font_Lambda([this]() {
+								if (ViewEntry->EntryType == EEntryType::Collection && ViewEntry->bIsDefaultCollection)
+								{
+									return FCoreStyle::Get().GetWidgetStyle<FTextBlockStyle>("Text.Large").Font;
+								}
+								else
+								{
+									return FCoreStyle::Get().GetWidgetStyle<FTextBlockStyle>("NormalText").Font;
+								}
+							})
 						]
 
 
@@ -381,7 +394,35 @@ void SPresetManager::Construct( const FArguments& InArgs )
 					.Value(0.4f)
 					.Resizable(false)
 						[
-							SNew(SVerticalBox)												
+							SNew(SVerticalBox)		
+								+ SVerticalBox::Slot()
+									.FillHeight(1.0f)
+								[
+			
+											SAssignNew(EditorPresetCollectionTreeView, STreeView<TSharedPtr<FPresetViewEntry> >)
+												.ItemHeight(32.0f)
+												.TreeItemsSource(&EditorCollectionsDataList)
+												.SelectionMode(ESelectionMode::Single)
+												.OnGenerateRow(this, &SPresetManager::HandleTreeGenerateRow)
+												.OnGetChildren(this, &SPresetManager::HandleTreeGetChildren)
+												.OnSelectionChanged(this, &SPresetManager::HandleUserTreeSelectionChanged)
+												.HeaderRow
+												(
+													SNew(SHeaderRow)
+													.Visibility(EVisibility::Collapsed)
+
+													+ SHeaderRow::Column("Collection")
+													.FixedWidth(150.0f)
+													.HeaderContent()
+													[
+														SNew(STextBlock)
+														.Text(LOCTEXT("PresetManagerCollectionTitleHeader", "Collection"))
+													]
+
+
+												)											
+								]
+
 								+ SVerticalBox::Slot()
 								.AutoHeight()
 								[
@@ -713,6 +754,16 @@ void SPresetManager::Construct( const FArguments& InArgs )
 		];
 
 		DeleteUserPresetButton->SetEnabled(false);
+		RegeneratePresetTrees();
+		if (UserCollectionsDataList.Num() == 0)
+		{
+			bAreUserCollectionsExpanded = false;
+		}
+		if (ProjectCollectionsDataList.Num() == 0)
+		{
+			bAreProjectCollectionsExpanded = false;
+		}
+
 }
 END_SLATE_FUNCTION_BUILD_OPTIMIZATION
 
@@ -721,6 +772,11 @@ SPresetManager::~SPresetManager()
 }
 
 void SPresetManager::Tick(const FGeometry& AllottedGeometry, const double InCurrentTime, const float InDeltaTime)
+{
+	RegeneratePresetTrees();
+}
+
+void SPresetManager::RegeneratePresetTrees()
 {
 	if (!ensure(UserSettings.IsValid()))
 	{
@@ -734,59 +790,73 @@ void SPresetManager::Tick(const FGeometry& AllottedGeometry, const double InCurr
 
 	TotalPresetCount = 0;
 
-	auto GenerateTreeEntries = [this](TArray<FSoftObjectPath>& AssetList,
+	auto GenerateSubTree = [this](UInteractiveToolsPresetCollectionAsset* PresetCollection, TSharedPtr<FPresetViewEntry> RootEntry)
+	{
+		TMap<FString, FInteractiveToolPresetStore >::TIterator ToolNameIter = PresetCollection->PerToolPresets.CreateIterator();
+		for (; (bool)ToolNameIter; ++ToolNameIter)
+		{
+			int32 ToolCount = 0;
+			for (int32 PresetIndex = 0; PresetIndex < ToolNameIter.Value().NamedPresets.Num(); ++PresetIndex)
+			{
+				ToolCount += ToolNameIter.Value().NamedPresets[PresetIndex].IsValid() ? 1 : 0;
+			}
+			if (ToolCount)
+			{
+				RootEntry->Children.Add(MakeShared<FPresetViewEntry>(
+					ToolNameIter.Value().ToolLabel,
+					ToolNameIter.Value().ToolIcon,
+					RootEntry->CollectionPath,
+					ToolNameIter.Key(),
+					ToolCount));
+				RootEntry->Children.Last()->Parent = RootEntry;
+				RootEntry->Count += ToolCount;
+				TotalPresetCount += ToolCount;
+			}
+		}
+
+	};
+
+	auto GenerateTreeEntries = [this, &GenerateSubTree](TObjectPtr<UInteractiveToolsPresetCollectionAsset> DefaultCollection,
+		TArray<FSoftObjectPath>* AssetList,
 		TArray< TSharedPtr< FPresetViewEntry > >& TreeList,
 		TSharedPtr<STreeView<TSharedPtr<FPresetViewEntry> > >& TreeView)
 	{
-		AssetList.RemoveAll([](const FSoftObjectPath& Path) { return !Path.IsAsset(); });
-		
 		bool bTreeNeedsRefresh = false;
 		TArray< TSharedPtr< FPresetViewEntry > > TempTreeDataList;
-		for (const FSoftObjectPath& Path : AssetList)
+
+		if (DefaultCollection)
 		{
-			UInteractiveToolsPresetCollectionAsset* PresetCollection = nullptr;
+			TSharedPtr<FPresetViewEntry> CollectionEntry = MakeShared<FPresetViewEntry>(
+				UserSettings->bDefaultCollectionEnabled,
+				FSoftObjectPath(),
+				DefaultCollection->CollectionLabel,
+				0);
+			CollectionEntry->bIsDefaultCollection = true;
+			GenerateSubTree(DefaultCollection, CollectionEntry);
+			TempTreeDataList.Add(CollectionEntry);
+		}
 
-			if (Path.IsAsset())
-			{
-				PresetCollection = Cast<UInteractiveToolsPresetCollectionAsset>(Path.TryLoad());
-			}
-			if (PresetCollection)
+		if (AssetList)
+		{
+			AssetList->RemoveAll([](const FSoftObjectPath& Path) { return !Path.IsAsset(); });
+
+			for (const FSoftObjectPath& Path : *AssetList)
 			{
-				if (Path == UserSettings->DefaultCollectionPath)
+				UInteractiveToolsPresetCollectionAsset* PresetCollection = nullptr;
+
+				if (Path.IsAsset())
 				{
-					TempTreeDataList.Add(MakeShared<FPresetViewEntry>(
-						UserSettings->bDefaultCollectionEnabled,
-						Path,
-						PresetCollection->CollectionLabel,
-						0));
+					PresetCollection = Cast<UInteractiveToolsPresetCollectionAsset>(Path.TryLoad());
 				}
-				else
+				if (PresetCollection)
 				{
-					TempTreeDataList.Add(MakeShared<FPresetViewEntry>(
+					TSharedPtr<FPresetViewEntry> CollectionEntry = MakeShared<FPresetViewEntry>(
 						UserSettings->EnabledPresetCollections.Contains(Path),
 						Path,
 						PresetCollection->CollectionLabel,
-						0));
-				}
-
-				TMap<FString, FInteractiveToolPresetStore >::TIterator ToolNameIter = PresetCollection->PerToolPresets.CreateIterator();
-				for (; (bool)ToolNameIter; ++ToolNameIter)
-				{
-					int32 ToolCount = 0;
-					for (int32 PresetIndex = 0; PresetIndex < ToolNameIter.Value().NamedPresets.Num(); ++PresetIndex)
-					{
-						ToolCount += ToolNameIter.Value().NamedPresets[PresetIndex].IsValid() ? 1 : 0;
-					}
-
-					TempTreeDataList.Last()->Children.Add(MakeShared<FPresetViewEntry>(
-						ToolNameIter.Value().ToolLabel,
-						ToolNameIter.Value().ToolIcon,
-						Path,
-						ToolNameIter.Key(),
-						ToolCount));
-					TempTreeDataList.Last()->Children.Last()->Parent = TempTreeDataList.Last();
-					TempTreeDataList.Last()->Count += ToolCount;
-					TotalPresetCount += ToolCount;
+						0);
+					GenerateSubTree(PresetCollection, CollectionEntry);
+					TempTreeDataList.Add(CollectionEntry);
 				}
 			}
 		}
@@ -816,7 +886,7 @@ void SPresetManager::Tick(const FGeometry& AllottedGeometry, const double InCurr
 		for (TSharedPtr<FPresetViewEntry>& Entry : TreeList)
 		{
 			Entry->bEnabled = UserSettings->EnabledPresetCollections.Contains(Entry->CollectionPath);
-			if (Entry->CollectionPath == UserSettings->DefaultCollectionPath)
+			if (Entry->bIsDefaultCollection)
 			{
 				Entry->bEnabled = UserSettings->bDefaultCollectionEnabled;
 			}
@@ -827,8 +897,17 @@ void SPresetManager::Tick(const FGeometry& AllottedGeometry, const double InCurr
 		}
 	};
 
-	GenerateTreeEntries(AvailablePresetCollections, ProjectCollectionsDataList, ProjectPresetCollectionTreeView);
-	GenerateTreeEntries(AvailableUserPresetCollections, UserCollectionsDataList, UserPresetCollectionTreeView);
+	// Handle the default collection
+	UPresetAssetSubsystem* PresetAssetSubsystem = GEditor->GetEditorSubsystem<UPresetAssetSubsystem>();
+	TObjectPtr<UInteractiveToolsPresetCollectionAsset> DefaultCollection = nullptr;
+	if (ensure(PresetAssetSubsystem))
+	{
+		DefaultCollection = PresetAssetSubsystem->GetDefaultCollection();
+	}
+
+	GenerateTreeEntries(nullptr, &AvailablePresetCollections, ProjectCollectionsDataList, ProjectPresetCollectionTreeView);
+	GenerateTreeEntries(nullptr, &AvailableUserPresetCollections, UserCollectionsDataList, UserPresetCollectionTreeView);
+	GenerateTreeEntries(DefaultCollection, nullptr, EditorCollectionsDataList, EditorPresetCollectionTreeView);
 
 }
 
@@ -858,6 +937,7 @@ void SPresetManager::GeneratePresetList(TSharedPtr<FPresetViewEntry> TreeEntry)
 	PresetDataList.Empty();
 	PresetListView->RequestListRefresh();
 	bHasActiveCollection = false;
+	ActivePresetToEdit = nullptr;
 
 	if (!TreeEntry)
 	{
@@ -867,12 +947,8 @@ void SPresetManager::GeneratePresetList(TSharedPtr<FPresetViewEntry> TreeEntry)
 	if (TreeEntry->EntryType == FPresetViewEntry::EEntryType::Collection ||
 		TreeEntry->EntryType == FPresetViewEntry::EEntryType::Tool)
 	{
-		UInteractiveToolsPresetCollectionAsset* PresetCollection = nullptr;
+		UInteractiveToolsPresetCollectionAsset* PresetCollection = GetCollectionFromEntry(TreeEntry);
 
-		if (TreeEntry->CollectionPath.IsAsset())
-		{
-			PresetCollection = Cast<UInteractiveToolsPresetCollectionAsset>(TreeEntry->CollectionPath.TryLoad());
-		}
 		if (PresetCollection)
 		{
 			if (TreeEntry->EntryType == FPresetViewEntry::EEntryType::Collection)
@@ -896,6 +972,7 @@ void SPresetManager::GeneratePresetList(TSharedPtr<FPresetViewEntry> TreeEntry)
 								ToolNameIter.Value().NamedPresets[PresetIndex].Tooltip,
 								FText::FromString(ToolNameIter.Value().NamedPresets[PresetIndex].Label)
 								));
+							PresetDataList.Last()->Parent = TreeEntry;
 							PresetDataList.Last()->CollectionPath = TreeEntry->CollectionPath;
 							PresetDataList.Last()->EntryIcon = ToolNameIter.Value().ToolIcon;
 						}
@@ -925,6 +1002,7 @@ void SPresetManager::GeneratePresetList(TSharedPtr<FPresetViewEntry> TreeEntry)
 							ToolData->NamedPresets[PresetIndex].Tooltip,
 							FText::FromString(ToolData->NamedPresets[PresetIndex].Label)
 							));
+						PresetDataList.Last()->Parent = TreeEntry;
 						PresetDataList.Last()->CollectionPath = TreeEntry->CollectionPath;
 						PresetDataList.Last()->EntryIcon = TreeEntry->EntryIcon;
 					}
@@ -954,7 +1032,7 @@ void SPresetManager::HandleUserTreeSelectionChanged(TSharedPtr<FPresetViewEntry>
 	DeleteUserPresetButton->SetEnabled(false);
 	if (TreeEntry)
 	{
-		if (TreeEntry->EntryType == FPresetViewEntry::EEntryType::Collection && TreeEntry->CollectionPath != UserSettings->DefaultCollectionPath)
+		if (TreeEntry->EntryType == FPresetViewEntry::EEntryType::Collection && !TreeEntry->bIsDefaultCollection)
 		{
 			DeleteUserPresetButton->SetEnabled(true);
 		}
@@ -1002,7 +1080,7 @@ void SPresetManager::SetCollectionEnabled(TSharedPtr<FPresetViewEntry> TreeEntry
 	{
 		return;
 	}
-	if (UserSettings->DefaultCollectionPath == TreeEntry->CollectionPath)
+	if (TreeEntry->bIsDefaultCollection)
 	{
 		UserSettings->bDefaultCollectionEnabled = (State == ECheckBoxState::Checked);
 		UserSettings->SaveEditorConfig();
@@ -1024,50 +1102,40 @@ void SPresetManager::SetCollectionEnabled(TSharedPtr<FPresetViewEntry> TreeEntry
 
 void SPresetManager::DeletePresetFromCollection(TSharedPtr< FPresetViewEntry > Entry)
 {
-	UInteractiveToolsPresetCollectionAsset* PresetCollection = nullptr;
-
-	if (Entry->CollectionPath.IsAsset())
-	{
-		PresetCollection = Cast<UInteractiveToolsPresetCollectionAsset>(Entry->CollectionPath.TryLoad());
-	}
+	UInteractiveToolsPresetCollectionAsset* PresetCollection = GetCollectionFromEntry(Entry);
 	if (PresetCollection)
 	{
 		PresetCollection->PerToolPresets[Entry->ToolName].NamedPresets.RemoveAt(Entry->PresetIndex);
 		PresetCollection->MarkPackageDirty();
 
-		PresetDataList.Remove(Entry);
-		PresetListView->RequestListRefresh();
+		GeneratePresetList(Entry->Parent);
 	}
+
+	SaveIfDefaultCollection(Entry);
 }
 
 void SPresetManager::SetPresetLabel(TSharedPtr< FPresetViewEntry > Entry, FText InLabel)
 {
-	UInteractiveToolsPresetCollectionAsset* PresetCollection = nullptr;
-
-	if (Entry->CollectionPath.IsAsset())
-	{
-		PresetCollection = Cast<UInteractiveToolsPresetCollectionAsset>(Entry->CollectionPath.TryLoad());
-	}
+	UInteractiveToolsPresetCollectionAsset* PresetCollection = GetCollectionFromEntry(Entry);
 	if (PresetCollection)
 	{
 		PresetCollection->PerToolPresets[Entry->ToolName].NamedPresets[Entry->PresetIndex].Label = InLabel.ToString();
 		PresetCollection->MarkPackageDirty();
 	}
+
+	SaveIfDefaultCollection(Entry);
 }
 
 void SPresetManager::SetPresetTooltip(TSharedPtr< FPresetViewEntry > Entry, FText InTooltip)
 {
-	UInteractiveToolsPresetCollectionAsset* PresetCollection = nullptr;
-
-	if (Entry->CollectionPath.IsAsset())
-	{
-		PresetCollection = Cast<UInteractiveToolsPresetCollectionAsset>(Entry->CollectionPath.TryLoad());
-	}
+	UInteractiveToolsPresetCollectionAsset* PresetCollection = GetCollectionFromEntry(Entry);
 	if (PresetCollection)
 	{
 		PresetCollection->PerToolPresets[Entry->ToolName].NamedPresets[Entry->PresetIndex].Tooltip = InTooltip.ToString();
 		PresetCollection->MarkPackageDirty();
 	}
+
+	SaveIfDefaultCollection(Entry);
 }
 
 void SPresetManager::DeleteSelectedUserPresetCollection()
@@ -1080,6 +1148,11 @@ void SPresetManager::DeleteSelectedUserPresetCollection()
 	if (SelectedUserCollections.Num() == 1)
 	{
 		TSharedPtr<FPresetViewEntry> Entry = SelectedUserCollections[0];
+		if (Entry->bIsDefaultCollection)
+		{
+			return;
+		}
+
 		FAssetData CollectionAsset;
 		if (AssetRegistryModule.Get().TryGetAssetByObjectPath(Entry->CollectionPath, CollectionAsset) == UE::AssetRegistry::EExists::Exists)
 		{
@@ -1162,4 +1235,35 @@ const FSlateBrush* SPresetManager::GetExpanderImage(TSharedPtr<SWidget> Expander
 	return FCoreStyle::Get().GetBrush(ResourceName);
 }
 
+UInteractiveToolsPresetCollectionAsset* SPresetManager::GetCollectionFromEntry(TSharedPtr<FPresetViewEntry> Entry)
+{
+	UInteractiveToolsPresetCollectionAsset* PresetCollection = nullptr;
+	UPresetAssetSubsystem* PresetAssetSubsystem = GEditor->GetEditorSubsystem<UPresetAssetSubsystem>();
+	
+	if (Entry->Root().bIsDefaultCollection && ensure(PresetAssetSubsystem))
+	{
+		PresetCollection = PresetAssetSubsystem->GetDefaultCollection();
+	}
+	else
+	{
+		if (Entry->CollectionPath.IsAsset())
+		{
+			PresetCollection = Cast<UInteractiveToolsPresetCollectionAsset>(Entry->CollectionPath.TryLoad());
+		}
+	}
+
+	return PresetCollection;
+}
+
+void SPresetManager::SaveIfDefaultCollection(TSharedPtr<FPresetViewEntry> Entry)
+{
+	UPresetAssetSubsystem* PresetAssetSubsystem = GEditor->GetEditorSubsystem<UPresetAssetSubsystem>();
+
+	if (Entry->Root().bIsDefaultCollection && ensure(PresetAssetSubsystem))
+	{
+		ensure(PresetAssetSubsystem->SaveDefaultCollection());
+	}
+}
+
+
 #undef LOCTEXT_NAMESPACE
\ No newline at end of file
diff --git a/Engine/Plugins/Experimental/ToolPresets/Source/PresetEditor/Private/Widgets/SPresetManager.h b/Engine/Plugins/Experimental/ToolPresets/Source/PresetEditor/Private/Widgets/SPresetManager.h
--- a/Engine/Plugins/Experimental/ToolPresets/Source/PresetEditor/Private/Widgets/SPresetManager.h
+++ b/Engine/Plugins/Experimental/ToolPresets/Source/PresetEditor/Private/Widgets/SPresetManager.h
@@ -19,7 +19,7 @@ class SVerticalBox;
 class SEditableTextBox;
 template<class ItemType> class STreeView;
 template<class ItemType> class SListView;
-
+class UInteractiveToolsPresetCollectionAsset;
 
 
 /**
@@ -55,6 +55,7 @@ public:
 		// Used for Collections/Tool Entries
 		bool bEnabled = false;
 		FSoftObjectPath CollectionPath;
+		bool bIsDefaultCollection = false;
 		FText EntryLabel;
 		FSlateBrush EntryIcon;
 		int32 Count = 0;
@@ -102,6 +103,7 @@ public:
 			bool bIsEqual = 
 				 EntryType == Other.EntryType &&
 				 CollectionPath == Other.CollectionPath &&
+				 bIsDefaultCollection == Other.bIsDefaultCollection &&
 				 Count == Other.Count &&
 				 ToolName.Equals(Other.ToolName) &&
 				 PresetIndex == Other.PresetIndex &&
@@ -121,6 +123,7 @@ public:
 			bool bIsEqual = bEnabled == Other.bEnabled &&
 				 EntryType == Other.EntryType &&
 				 CollectionPath == Other.CollectionPath &&
+				 bIsDefaultCollection == Other.bIsDefaultCollection && 
 				 Count == Other.Count &&
 				 EntryLabel.EqualTo(Other.EntryLabel) &&
 				 ToolName.Equals(Other.ToolName) &&
@@ -136,6 +139,16 @@ public:
 
 			return bIsEqual;
 		}
+
+		FPresetViewEntry& Root()
+		{
+			FPresetViewEntry* ActiveNode = this;
+			while (ActiveNode->Parent)
+			{
+				ActiveNode = ActiveNode->Parent.Get();
+			}
+			return *ActiveNode;
+		}
 	};
 
 
@@ -160,6 +173,8 @@ protected:
 
 private:
 
+	void RegeneratePresetTrees();
+
 	int32 GetTotalPresetCount() const;
 
 	TSharedRef<ITableRow> HandleTreeGenerateRow(TSharedPtr<FPresetViewEntry> TreeEntry, const TSharedRef<STableViewBase>& OwnerTable);
@@ -186,6 +201,9 @@ private:
 	const FSlateBrush* GetUserCollectionsExpanderImage() const;
 	const FSlateBrush* GetExpanderImage(TSharedPtr<SWidget> ExpanderWidget, bool bIsUserCollections) const;
 
+	UInteractiveToolsPresetCollectionAsset* GetCollectionFromEntry(TSharedPtr<FPresetViewEntry> Entry);
+	void SaveIfDefaultCollection(TSharedPtr<FPresetViewEntry> Entry);
+
 private:
 
 	TWeakObjectPtr<UPresetUserSettings> UserSettings;
@@ -200,6 +218,10 @@ private:
 	TArray< TSharedPtr< FPresetViewEntry > > UserCollectionsDataList;
 	TSharedPtr<STreeView<TSharedPtr<FPresetViewEntry> > > UserPresetCollectionTreeView;
 
+	TArray< TSharedPtr< FPresetViewEntry > > EditorCollectionsDataList;
+	TSharedPtr<STreeView<TSharedPtr<FPresetViewEntry> > > EditorPresetCollectionTreeView;
+
+
 	TArray< TSharedPtr< FPresetViewEntry > > PresetDataList;
 	TSharedPtr<SListView<TSharedPtr<FPresetViewEntry> > > PresetListView;
 
diff --git a/Engine/Plugins/Experimental/ToolPresets/Source/PresetEditor/Public/PresetSettings.h b/Engine/Plugins/Experimental/ToolPresets/Source/PresetEditor/Public/PresetSettings.h
--- a/Engine/Plugins/Experimental/ToolPresets/Source/PresetEditor/Public/PresetSettings.h
+++ b/Engine/Plugins/Experimental/ToolPresets/Source/PresetEditor/Public/PresetSettings.h
@@ -19,14 +19,13 @@ public:
 	GENERATED_BODY()
 
 	UPROPERTY(EditAnywhere, Category = "Interactive Tool Presets|Collections", meta = (EditorConfig, AllowedClasses = "/Script/PresetAsset.InteractiveToolsPresetCollectionAsset"))
-	TSet<FSoftObjectPath> EnabledPresetCollections;// = { FSoftObjectPath(FTopLevelAssetPath("/ToolPresets/Presets/_DefaultCollection", "_DefaultCollection")) };
+	TSet<FSoftObjectPath> EnabledPresetCollections;
 
-	// TODO: Ideally the above property would be able to store if the default collection was enabled or not.
-	// However, currently the EditorConfig system has some trouble with defaults and array/set properties.
-	// To overcome this, we store it separately which doesn't seem to break as badly.
-
-	UPROPERTY(VisibleAnywhere, Category = "Interactive Tool Presets|Collections", meta = (EditorConfig, AllowedClasses = "/Script/PresetAsset.InteractiveToolsPresetCollectionAsset"))
-	FSoftObjectPath DefaultCollectionPath = FSoftObjectPath(FTopLevelAssetPath("/ToolPresets/Presets/_DefaultCollection", "_DefaultCollection"));
+	//~ Ideally the above property would be able to store if the default collection was enabled or not.
+	//~ 
+	//~ However, the default collection itself is stored via an alternative JSON representation, accessed through the PresetAssetSubsystem,
+	//~ to avoid issues with automatic asset generation. Therefore it doesn't have a "path" in the traditional sense that the other collections
+	//~ do, requiring a separate tracking of it's enabled/disabled status.
 
 	UPROPERTY(EditAnywhere, Category = "Interactive Tool Presets|Collections", meta = (EditorConfig))
 	bool bDefaultCollectionEnabled = true;
@@ -40,7 +39,7 @@ private:
 
 
 /**
- * Implements the settings for the PresetEditor.
+ * Implements the settings for the Project Preset Collections.
  */
 UCLASS(config = Editor)
 class PRESETEDITOR_API UPresetProjectSettings
@@ -63,4 +62,4 @@ public:
 	/* Controls which preset collection assets are to be loaded for this project.  */
 	UPROPERTY(config, EditAnywhere, Category = "Interactive Tool Presets|Collections", meta = (AllowedClasses = "/Script/PresetAsset.InteractiveToolsPresetCollectionAsset"))
 	TSet<FSoftObjectPath> LoadedPresetCollections;
-};
\ No newline at end of file
+};
