ToolPresets: Remove references to unneeded module loading that was causing crashes.

diff --git a/Engine/Plugins/Experimental/ToolPresets/Source/PresetEditor/Private/PresetEditorModule.cpp b/Engine/Plugins/Experimental/ToolPresets/Source/PresetEditor/Private/PresetEditorModule.cpp
--- a/Engine/Plugins/Experimental/ToolPresets/Source/PresetEditor/Private/PresetEditorModule.cpp
+++ b/Engine/Plugins/Experimental/ToolPresets/Source/PresetEditor/Private/PresetEditorModule.cpp
@@ -32,11 +32,6 @@ public:
 
 	virtual void StartupModule() override
 	{
-		FModuleManager::Get().LoadModuleChecked(TEXT("DataValidation"));
-		FModuleManager::Get().LoadModuleChecked(TEXT("AssetTools"));
-		FModuleManager::Get().LoadModuleChecked(TEXT("ContentBrowser"));
-		FModuleManager::Get().LoadModuleChecked(TEXT("AssetRegistry"));
-
 		FGlobalTabmanager::Get()->RegisterNomadTabSpawner(PresetEditorTabName, FOnSpawnTab::CreateRaw(this, &FPresetEditorModule::HandleSpawnPresetEditorTab))
 			.SetDisplayName(NSLOCTEXT("FPresetModule", "PresetTabTitle", "Preset Manager"))
 			.SetTooltipText(NSLOCTEXT("FPresetModule", "PresetTooltipText", "Open the Preset Manager tab."))
@@ -44,8 +39,6 @@ public:
 			.SetAutoGenerateMenuEntry(false);
 
 		FCoreDelegates::OnPostEngineInit.AddRaw(this, &FPresetEditorModule::OnPostEngineInit);
-
-		
 	}
 	
 	virtual void ShutdownModule() override
