Tool Presets: Change handling of the User vs Editor preset lists in the Preset Manager to separate selection handling between the two sections.

diff --git a/Engine/Plugins/Experimental/ToolPresets/Source/PresetEditor/Private/Widgets/SPresetManager.cpp b/Engine/Plugins/Experimental/ToolPresets/Source/PresetEditor/Private/Widgets/SPresetManager.cpp
--- a/Engine/Plugins/Experimental/ToolPresets/Source/PresetEditor/Private/Widgets/SPresetManager.cpp
+++ b/Engine/Plugins/Experimental/ToolPresets/Source/PresetEditor/Private/Widgets/SPresetManager.cpp
@@ -405,7 +405,7 @@ void SPresetManager::Construct( const FArguments& InArgs )
 												.SelectionMode(ESelectionMode::Single)
 												.OnGenerateRow(this, &SPresetManager::HandleTreeGenerateRow)
 												.OnGetChildren(this, &SPresetManager::HandleTreeGetChildren)
-												.OnSelectionChanged(this, &SPresetManager::HandleUserTreeSelectionChanged)
+												.OnSelectionChanged(this, &SPresetManager::HandleEditorTreeSelectionChanged)
 												.HeaderRow
 												(
 													SNew(SHeaderRow)
@@ -1013,11 +1013,20 @@ void SPresetManager::GeneratePresetList(TSharedPtr<FPresetViewEntry> TreeEntry)
 
 }
 
+void SPresetManager::HandleEditorTreeSelectionChanged(TSharedPtr<FPresetViewEntry> TreeEntry, ESelectInfo::Type SelectInfo)
+{
+	if (SelectInfo != ESelectInfo::Direct) {
+		UserPresetCollectionTreeView->ClearSelection();
+		ProjectPresetCollectionTreeView->ClearSelection();
+		GeneratePresetList(TreeEntry);
+	}
+}
 
 void SPresetManager::HandleTreeSelectionChanged(TSharedPtr<FPresetViewEntry> TreeEntry, ESelectInfo::Type SelectInfo)
 {
 	if (SelectInfo != ESelectInfo::Direct) {
 		UserPresetCollectionTreeView->ClearSelection();
+		EditorPresetCollectionTreeView->ClearSelection();
 		GeneratePresetList(TreeEntry);
 	}
 }
@@ -1026,6 +1035,7 @@ void SPresetManager::HandleUserTreeSelectionChanged(TSharedPtr<FPresetViewEntry>
 {
 	if (SelectInfo != ESelectInfo::Direct) {
 		ProjectPresetCollectionTreeView->ClearSelection();
+		EditorPresetCollectionTreeView->ClearSelection();
 		GeneratePresetList(TreeEntry);
 	}
 
diff --git a/Engine/Plugins/Experimental/ToolPresets/Source/PresetEditor/Private/Widgets/SPresetManager.h b/Engine/Plugins/Experimental/ToolPresets/Source/PresetEditor/Private/Widgets/SPresetManager.h
--- a/Engine/Plugins/Experimental/ToolPresets/Source/PresetEditor/Private/Widgets/SPresetManager.h
+++ b/Engine/Plugins/Experimental/ToolPresets/Source/PresetEditor/Private/Widgets/SPresetManager.h
@@ -181,6 +181,7 @@ private:
 	void HandleTreeGetChildren(TSharedPtr<FPresetViewEntry> TreeEntry, TArray< TSharedPtr<FPresetViewEntry> >& ChildrenOut);
 	void HandleTreeSelectionChanged(TSharedPtr<FPresetViewEntry> TreeEntry, ESelectInfo::Type);
 	void HandleUserTreeSelectionChanged(TSharedPtr<FPresetViewEntry> TreeEntry, ESelectInfo::Type);
+	void HandleEditorTreeSelectionChanged(TSharedPtr<FPresetViewEntry> TreeEntry, ESelectInfo::Type);
 	EVisibility ProjectPresetCollectionsVisibility() const;
 
 	void GeneratePresetList(TSharedPtr<FPresetViewEntry> TreeEntry);
