Tool Presets: Adjust handling of the preset label and tooltip fields to cap character counts that are more reasonable and don't break integration into menus.

diff --git a/Engine/Plugins/Editor/ModelingToolsEditorMode/Source/ModelingToolsEditorMode/Private/ModelingToolsEditorModeToolkit.cpp b/Engine/Plugins/Editor/ModelingToolsEditorMode/Source/ModelingToolsEditorMode/Private/ModelingToolsEditorModeToolkit.cpp
--- a/Engine/Plugins/Editor/ModelingToolsEditorMode/Source/ModelingToolsEditorMode/Private/ModelingToolsEditorModeToolkit.cpp
+++ b/Engine/Plugins/Editor/ModelingToolsEditorMode/Source/ModelingToolsEditorMode/Private/ModelingToolsEditorModeToolkit.cpp
@@ -1029,7 +1029,8 @@ TSharedRef<SWidget> FModelingToolsEditorModeToolkit::GetPresetCreateButtonConten
 				+ SHorizontalBox::Slot()
 				[
 					SNew(SEditableTextBox)
-					.OnTextCommitted_Lambda([this](const FText& NewLabel, ETextCommit::Type&) { NewPresetLabel = NewLabel.ToString(); })
+					// Cap the number of characters sent out of the text box, so we don't overflow menus and tooltips
+					.OnTextCommitted_Lambda([this](const FText& NewLabel, ETextCommit::Type&) { NewPresetLabel = NewLabel.ToString().Left(255); })
 					.OverflowPolicy(ETextOverflowPolicy::Ellipsis)
 					.ToolTipText(LOCTEXT("ToolPresets_CreatePresetLabel_Tooltip", "A short, descriptive identifier for the new preset."))
 				]
@@ -1047,7 +1048,8 @@ TSharedRef<SWidget> FModelingToolsEditorModeToolkit::GetPresetCreateButtonConten
 				+ SHorizontalBox::Slot()
 				[
 					SNew(SEditableTextBox)
-					.OnTextCommitted_Lambda([this](const FText& NewToolTip, ETextCommit::Type&) { NewPresetTooltip = NewToolTip.ToString(); })
+					// Cap the number of characters sent out of the text box, so we don't overflow menus and tooltips
+					.OnTextCommitted_Lambda([this](const FText& NewToolTip, ETextCommit::Type&) { NewPresetTooltip = NewToolTip.ToString().Left(2048); })
 					.OverflowPolicy(ETextOverflowPolicy::Ellipsis)
 					.ToolTipText(LOCTEXT("ToolPresets_CreatePresetTooltip_Tooltip", "A descriptive tooltip for the new preset."))
 				]
diff --git a/Engine/Plugins/Experimental/ToolPresets/Source/PresetEditor/Private/Widgets/SPresetManager.cpp b/Engine/Plugins/Experimental/ToolPresets/Source/PresetEditor/Private/Widgets/SPresetManager.cpp
--- a/Engine/Plugins/Experimental/ToolPresets/Source/PresetEditor/Private/Widgets/SPresetManager.cpp
+++ b/Engine/Plugins/Experimental/ToolPresets/Source/PresetEditor/Private/Widgets/SPresetManager.cpp
@@ -679,7 +679,8 @@ void SPresetManager::Construct( const FArguments& InArgs )
 								.FillWidth(1.f)								
 								.Padding(5.0f)
 								[
-									SNew(SEditableTextBox)			
+									SNew(SEditableTextBox)
+									.OverflowPolicy(ETextOverflowPolicy::Ellipsis)
 									.Text_Lambda([this]() {
 										if (ActivePresetToEdit)
 										{
@@ -690,13 +691,15 @@ void SPresetManager::Construct( const FArguments& InArgs )
 									.OnTextChanged_Lambda([this](const FText& NewText) {
 										if (ActivePresetToEdit)
 										{
-											ActivePresetToEdit->PresetLabel = NewText.ToString();
+											// Cap the number of characters sent out of the text box, so we don't overflow menus and tooltips
+											ActivePresetToEdit->PresetLabel = NewText.ToString().Left(255);
 										}
 									})
 									.OnTextCommitted_Lambda([this](const FText& NewText, ETextCommit::Type CommitStatus){
 										if (ActivePresetToEdit)
 										{
-											SetPresetLabel(ActivePresetToEdit, NewText);
+											// Cap the number of characters sent out of the text box, so we don't overflow menus and tooltips
+											SetPresetLabel(ActivePresetToEdit, FText::FromString(NewText.ToString().Left(255)));
 										}
 									})
 								]
@@ -722,10 +725,12 @@ void SPresetManager::Construct( const FArguments& InArgs )
 									.MinDesiredHeight(44.f)
 									.MaxDesiredHeight(44.0f)
 									[
-										SNew(SMultiLineEditableTextBox)											
+										SNew(SMultiLineEditableTextBox)		
+										
 										.AllowMultiLine(false)
 										.AutoWrapText(true)
 										.WrappingPolicy(ETextWrappingPolicy::DefaultWrapping)
+										.OverflowPolicy(ETextOverflowPolicy::Ellipsis)
 										.Text_Lambda([this]() {
 											if (ActivePresetToEdit)
 											{
@@ -736,13 +741,15 @@ void SPresetManager::Construct( const FArguments& InArgs )
 										.OnTextChanged_Lambda([this](const FText& NewText) {
 											if (ActivePresetToEdit)
 											{
-												ActivePresetToEdit->PresetTooltip = NewText.ToString();
+												// Cap the number of characters sent out of the text box, so we don't overflow menus and tooltips
+												ActivePresetToEdit->PresetTooltip = NewText.ToString().Left(2048);
 											}
 										})
 										.OnTextCommitted_Lambda([this](const FText& NewText, ETextCommit::Type CommitStatus){
 											if (ActivePresetToEdit)
 											{
-												SetPresetTooltip(ActivePresetToEdit, NewText);
+												// Cap the number of characters sent out of the text box, so we don't overflow menus and tooltips
+												SetPresetTooltip(ActivePresetToEdit, FText::FromString(NewText.ToString().Left(2048)));
 											}
 										})
 									]
