Tool Presets: Fix preset serialization to skip properties marked as transient or do not serialize.

diff --git a/Engine/Plugins/Experimental/ToolPresets/Source/PresetAsset/Private/PresetAsset.cpp b/Engine/Plugins/Experimental/ToolPresets/Source/PresetAsset/Private/PresetAsset.cpp
--- a/Engine/Plugins/Experimental/ToolPresets/Source/PresetAsset/Private/PresetAsset.cpp
+++ b/Engine/Plugins/Experimental/ToolPresets/Source/PresetAsset/Private/PresetAsset.cpp
@@ -15,10 +15,17 @@ void FInteractiveToolPresetDefintion::SetStoredPropertyData(TArray<UObject*>& Pr
 		{
 			for (FProperty* Prop : TFieldRange<FProperty>(PropertySet->GetClass()))
 			{
-				TSharedPtr<FJsonValue> JsonProp = FJsonObjectConverter::UPropertyToJsonValue(Prop, Prop->ContainerPtrToValuePtr<void>(PropertySet));				
-				FString FieldName;
-				Prop->GetName(FieldName);
-				PropertyJsonObject->SetField(FieldName, JsonProp);
+#if WITH_EDITOR
+				if (! (Prop->HasMetaData(TEXT("TransientToolProperty")) ||
+					   Prop->GetPropertyFlags() & EPropertyFlags::CPF_SkipSerialization ||
+					   Prop->GetPropertyFlags() & EPropertyFlags::CPF_Transient) )
+#endif
+				{
+					TSharedPtr<FJsonValue> JsonProp = FJsonObjectConverter::UPropertyToJsonValue(Prop, Prop->ContainerPtrToValuePtr<void>(PropertySet));
+					FString FieldName;
+					Prop->GetName(FieldName);
+					PropertyJsonObject->SetField(FieldName, JsonProp);
+				}
 			}			
 		}
 
