Tool Presets: Correctly trigger OnModified for interactive tools when presets load, allowing active tool state to reflect the changed property values.

diff --git a/Engine/Plugins/Experimental/ToolPresets/Source/PresetAsset/Private/PresetAsset.cpp b/Engine/Plugins/Experimental/ToolPresets/Source/PresetAsset/Private/PresetAsset.cpp
--- a/Engine/Plugins/Experimental/ToolPresets/Source/PresetAsset/Private/PresetAsset.cpp
+++ b/Engine/Plugins/Experimental/ToolPresets/Source/PresetAsset/Private/PresetAsset.cpp
@@ -57,7 +57,11 @@ void FInteractiveToolPresetDefintion::LoadStoredPropertyData(TArray<UObject*>& P
 				TSharedPtr<FJsonValue>* JsonField = PropertyJsonObject->Values.Find(FieldName);
 				if (JsonField)
 				{
-					FJsonObjectConverter::JsonValueToUProperty(*JsonField, Prop, Prop->ContainerPtrToValuePtr<void>(PropertySet));
+					if (FJsonObjectConverter::JsonValueToUProperty(*JsonField, Prop, Prop->ContainerPtrToValuePtr<void>(PropertySet)))
+					{
+						FPropertyChangedEvent ChangeEvent(Prop, EPropertyChangeType::ValueSet);
+						PropertySet->PostEditChangeProperty(ChangeEvent);
+					}
 				}
 			}
 		}
