Tool Presets: Fixes a number of UI adjustments and fixes for the Preset feature in modeling tools.

* Provides support for displaying asset names in the asset picker flyout's tile view.
* Reworks Preset manager moderately to have more streamlined controls and use of context menus to handle deletion and renaming
* Fixes the asset picker to display preset collections regardless if the user has configured the editor to show plugin/editor content
* Alters some language around the preset system for clarity

diff --git a/Engine/Plugins/Editor/ModelingToolsEditorMode/Source/ModelingToolsEditorMode/Private/ModelingToolsEditorModeToolkit.cpp b/Engine/Plugins/Editor/ModelingToolsEditorMode/Source/ModelingToolsEditorMode/Private/ModelingToolsEditorModeToolkit.cpp
--- a/Engine/Plugins/Editor/ModelingToolsEditorMode/Source/ModelingToolsEditorMode/Private/ModelingToolsEditorModeToolkit.cpp
+++ b/Engine/Plugins/Editor/ModelingToolsEditorMode/Source/ModelingToolsEditorMode/Private/ModelingToolsEditorModeToolkit.cpp
@@ -1081,6 +1081,9 @@ TSharedRef<SWidget> FModelingToolsEditorModeToolkit::GetPresetCreateButtonConten
 					.InitiallySelectedAsset(*CurrentPreset)
 					.FlyoutTileSize(FVector2D(80, 80))
 					.ComboButtonTileSize(FVector2D(80, 80))
+					.AssetThumbnailLabel(EThumbnailLabel::AssetName)
+					.bForceShowPluginContent(true)
+					.bForceShowEngineContent(true)
 				]
 				+ SHorizontalBox::Slot()
 				[
@@ -1101,7 +1104,7 @@ TSharedRef<SWidget> FModelingToolsEditorModeToolkit::GetPresetCreateButtonConten
 						.Text_Lambda([this](){
 						if (CurrentPresetLabel.IsEmpty())
 						{
-							return LOCTEXT("NewPresetNoCollectionSpecifiedMessage", "None - Preset will be added to Editor Presets.");
+							return LOCTEXT("NewPresetNoCollectionSpecifiedMessage", "None - Preset will be added to the default Personal Presets Collection.");
 						}
 						else {
 							return CurrentPresetLabel;
@@ -1111,8 +1114,6 @@ TSharedRef<SWidget> FModelingToolsEditorModeToolkit::GetPresetCreateButtonConten
 				]					
 			],
 			FolderDialogArguments, true);
-
-		return FReply::Handled();
 	};
 
 	constexpr bool bShouldCloseMenuAfterSelection = true;
@@ -1162,59 +1163,24 @@ TSharedRef<SWidget> FModelingToolsEditorModeToolkit::GetPresetCreateButtonConten
 		}
 
 		MenuBuilder.BeginSection(NoExtensionHook, LOCTEXT("ModelingPresetPanelHeaderManagePresets", "Manage Presets"));
-		MenuBuilder.AddWidget(
-			SNew(SButton)
-			.OnClicked(FOnClicked::CreateLambda(OpenNewPresetDialog))
-			.Content()
-			[
-				SNew(SHorizontalBox)
-				+ SHorizontalBox::Slot()
-			.AutoWidth()
-			.VAlign(VAlign_Center)
-			.Padding(4.0, 0.0f)
-			[
-				SNew(SImage)
-				.ColorAndOpacity(FSlateColor::UseForeground())
-			.Image(FAppStyle::Get().GetBrush("Icons.Plus"))
-			]
-		+ SHorizontalBox::Slot()
-			.AutoWidth()
-			.VAlign(VAlign_Center)
-			.Padding(4.0, 0.0f)
-			[
-				SNew(STextBlock)
-				.Text(LOCTEXT("ModelingPresetPanelCreateNewPreset", "Create New Preset"))
-				.Justification(ETextJustify::Center)
-				.Font(FAppStyle::Get().GetFontStyle("EditorModesPanel.CategoryFontStyle"))
-			]
-			], FText());
 
-		MenuBuilder.AddWidget(
-			SNew(SButton)
-			.OnClicked_Lambda([this]() {IPresetEditorModule::Get().ExecuteOpenPresetEditor(); return FReply::Handled(); })
-			.Content()
-			[
-				SNew(SHorizontalBox)
-				+ SHorizontalBox::Slot()
-				.AutoWidth()
-				.VAlign(VAlign_Center)
-				.Padding(4.0, 0.0f)
-				[
-					SNew(SImage)
-					.ColorAndOpacity(FSlateColor::UseForeground())
-					.Image(FAppStyle::Get().GetBrush("Icons.Settings"))
-				]
-				+ SHorizontalBox::Slot()
-				.AutoWidth()
-				.VAlign(VAlign_Center)
-				.Padding(4.0, 0.0f)
-				[
-					SNew(STextBlock)
-					.Text(LOCTEXT("ModelingPresetPanelOpenPresetManager", "Preset Manager"))
-					.Justification(ETextJustify::Center)
-					.Font(FAppStyle::Get().GetFontStyle("EditorModesPanel.CategoryFontStyle"))
-				]
-			], FText());
+		FUIAction CreateNewPresetAction;
+		CreateNewPresetAction.ExecuteAction = FExecuteAction::CreateLambda(OpenNewPresetDialog);
+		MenuBuilder.AddMenuEntry(
+			FText(LOCTEXT("ModelingPresetPanelCreateNewPreset", "Create New Preset")),
+			FText(LOCTEXT("ModelingPresetPanelCreateNewPresetTooltip", "Create New Preset in specified Collection")),
+			FSlateIcon( FAppStyle::Get().GetStyleSetName(), "Icons.Plus"),
+			CreateNewPresetAction);
+
+		FUIAction OpenPresetManangerAction;
+		OpenPresetManangerAction.ExecuteAction = FExecuteAction::CreateLambda([this]() {
+			IPresetEditorModule::Get().ExecuteOpenPresetEditor();
+		});
+		MenuBuilder.AddMenuEntry(
+			FText(LOCTEXT("ModelingPresetPanelOpenPresetMananger", "Manage Presets...")),
+			FText(LOCTEXT("ModelingPresetPanelOpenPresetManagerTooltip", "Open Preset Manager to manage presets and their collections")),
+			FSlateIcon(FAppStyle::Get().GetStyleSetName(), "Icons.Settings"),
+			OpenPresetManangerAction);
 
 		MenuBuilder.EndSection();
 
diff --git a/Engine/Plugins/Experimental/MeshModelingToolsetExp/Source/ModelingEditorUI/Private/ModelingWidgets/SToolInputAssetComboPanel.cpp b/Engine/Plugins/Experimental/MeshModelingToolsetExp/Source/ModelingEditorUI/Private/ModelingWidgets/SToolInputAssetComboPanel.cpp
--- a/Engine/Plugins/Experimental/MeshModelingToolsetExp/Source/ModelingEditorUI/Private/ModelingWidgets/SToolInputAssetComboPanel.cpp
+++ b/Engine/Plugins/Experimental/MeshModelingToolsetExp/Source/ModelingEditorUI/Private/ModelingWidgets/SToolInputAssetComboPanel.cpp
@@ -116,7 +116,11 @@ void SToolInputAssetComboPanel::Construct(const FArguments& InArgs)
 		];
 
 
-	ComboButton->SetOnGetMenuContent(FOnGetContent::CreateLambda([this]() 
+	EThumbnailLabel::Type TileThumbnailLabel = InArgs._AssetThumbnailLabel;
+	bool bForceShowEngineContent = InArgs._bForceShowEngineContent;
+	bool bForceShowPluginContent = InArgs._bForceShowPluginContent;
+
+	ComboButton->SetOnGetMenuContent(FOnGetContent::CreateLambda([this, TileThumbnailLabel,bForceShowEngineContent, bForceShowPluginContent]()
 	{
 		// Configure filter for asset picker
 		FAssetPickerConfig Config;
@@ -129,8 +133,10 @@ void SToolInputAssetComboPanel::Construct(const FArguments& InArgs)
 		Config.bFocusSearchBoxWhenOpened = true;
 		Config.bAllowNullSelection = true;
 		Config.bAllowDragging = false;
-		Config.ThumbnailLabel = EThumbnailLabel::Type::NoLabel;
+		Config.ThumbnailLabel = TileThumbnailLabel;
 		Config.OnAssetSelected = FOnAssetSelected::CreateSP(this, &SToolInputAssetComboPanel::NewAssetSelected);
+		Config.bForceShowEngineContent = bForceShowEngineContent;
+		Config.bForceShowPluginContent = bForceShowPluginContent;
 
 		// build asset picker UI
 		TSharedRef<SToolInputAssetPicker> AssetPickerWidget = SNew( SToolInputAssetPicker )
diff --git a/Engine/Plugins/Experimental/MeshModelingToolsetExp/Source/ModelingEditorUI/Public/ModelingWidgets/SToolInputAssetComboPanel.h b/Engine/Plugins/Experimental/MeshModelingToolsetExp/Source/ModelingEditorUI/Public/ModelingWidgets/SToolInputAssetComboPanel.h
--- a/Engine/Plugins/Experimental/MeshModelingToolsetExp/Source/ModelingEditorUI/Public/ModelingWidgets/SToolInputAssetComboPanel.h
+++ b/Engine/Plugins/Experimental/MeshModelingToolsetExp/Source/ModelingEditorUI/Public/ModelingWidgets/SToolInputAssetComboPanel.h
@@ -9,6 +9,7 @@
 #include "PropertyHandle.h"
 #include "AssetRegistry/AssetData.h"
 #include "CollectionManagerTypes.h"
+#include "AssetThumbnail.h"
 
 class SComboButton;
 class FAssetThumbnailPool;
@@ -66,6 +67,9 @@ public:
 		, _AssetClassType(0)
 		, _OnSelectionChanged()
 		, _InitiallySelectedAsset()
+		, _AssetThumbnailLabel(EThumbnailLabel::NoLabel)
+		, _bForceShowEngineContent(false)
+		, _bForceShowPluginContent(false)
 	{
 		}
 
@@ -99,6 +103,15 @@ public:
 		/** Sets the asset selected by the widget before any user made selection occurs. */
 		SLATE_ARGUMENT( FAssetData, InitiallySelectedAsset)
 
+		/** Sets the type of label used for the asset picker tiles */
+		SLATE_ARGUMENT( EThumbnailLabel::Type, AssetThumbnailLabel)
+
+		/** Indicates if engine content should always be shown */
+		SLATE_ARGUMENT(bool, bForceShowEngineContent)
+
+		/** Indicates if plugin content should always be shown */
+		SLATE_ARGUMENT(bool, bForceShowPluginContent)
+
 	SLATE_END_ARGS()
 
 
diff --git a/Engine/Plugins/Experimental/ToolPresets/Source/PresetAsset/Private/PresetAssetSubsystem.cpp b/Engine/Plugins/Experimental/ToolPresets/Source/PresetAsset/Private/PresetAssetSubsystem.cpp
--- a/Engine/Plugins/Experimental/ToolPresets/Source/PresetAsset/Private/PresetAssetSubsystem.cpp
+++ b/Engine/Plugins/Experimental/ToolPresets/Source/PresetAsset/Private/PresetAssetSubsystem.cpp
@@ -49,7 +49,7 @@ void UPresetAssetSubsystem::InitializeDefaultCollection()
 	*/
 
 	DefaultCollection = NewObject<UInteractiveToolsPresetCollectionAsset>();
-	DefaultCollection->CollectionLabel = LOCTEXT("DefaultCollectionLabel", "Editor Presets");
+	DefaultCollection->CollectionLabel = LOCTEXT("DefaultCollectionLabel", "Personal Presets (Default)");
 	DefaultCollection->LoadEditorConfig();
 }
 
diff --git a/Engine/Plugins/Experimental/ToolPresets/Source/PresetEditor/PresetEditor.Build.cs b/Engine/Plugins/Experimental/ToolPresets/Source/PresetEditor/PresetEditor.Build.cs
--- a/Engine/Plugins/Experimental/ToolPresets/Source/PresetEditor/PresetEditor.Build.cs
+++ b/Engine/Plugins/Experimental/ToolPresets/Source/PresetEditor/PresetEditor.Build.cs
@@ -30,6 +30,7 @@ namespace UnrealBuildTool.Rules
 					"ContentBrowser",
 					"ContentBrowserData",
 					"EditorConfig",
+					"AssetTools",
 				}
 			);
 		}
diff --git a/Engine/Plugins/Experimental/ToolPresets/Source/PresetEditor/Private/Widgets/SPresetManager.cpp b/Engine/Plugins/Experimental/ToolPresets/Source/PresetEditor/Private/Widgets/SPresetManager.cpp
--- a/Engine/Plugins/Experimental/ToolPresets/Source/PresetEditor/Private/Widgets/SPresetManager.cpp
+++ b/Engine/Plugins/Experimental/ToolPresets/Source/PresetEditor/Private/Widgets/SPresetManager.cpp
@@ -16,6 +16,7 @@
 #include "Widgets/Input/SComboButton.h"
 #include "Widgets/Input/SCheckBox.h"
 #include "Widgets/Input/SEditableTextBox.h"
+#include "Widgets/Text/SInlineEditableTextBlock.h"
 #include "Widgets/Input/SMultiLineEditableTextBox.h"
 #include "Editor/TransBuffer.h"
 #include "PresetSettings.h"
@@ -30,10 +31,13 @@
 #include "PresetAsset.h"
 #include "SNegativeActionButton.h"
 #include "SPositiveActionButton.h"
+#include "SSimpleButton.h"
 #include "PresetEditorStyle.h"
 #include "PresetAssetSubsystem.h"
-
-
+#include "Framework/Commands/GenericCommands.h"
+#include "Modules/ModuleManager.h"
+#include "ISettingsModule.h"
+#include "IAssetTools.h"
 
 #define LOCTEXT_NAMESPACE "SPresetManager"
 
@@ -41,6 +45,8 @@ DECLARE_DELEGATE_TwoParams(FOnCollectionEnabledCheckboxChanged, TSharedPtr<SPres
 DECLARE_DELEGATE_TwoParams(FOnPresetLabelChanged, TSharedPtr<SPresetManager::FPresetViewEntry>, FText)
 DECLARE_DELEGATE_TwoParams(FOnPresetTooltipChanged, TSharedPtr<SPresetManager::FPresetViewEntry>, FText)
 DECLARE_DELEGATE_OneParam(FOnPresetDeleted, TSharedPtr<SPresetManager::FPresetViewEntry>)
+DECLARE_DELEGATE_TwoParams(FOnCollectionRenameStarted, TSharedPtr<SPresetManager::FPresetViewEntry>, TSharedPtr<SEditableTextBox> RenameWidget)
+DECLARE_DELEGATE_TwoParams(FOnCollectionRenameEnded, TSharedPtr<SPresetManager::FPresetViewEntry>, const FText& NewText)
 
 namespace PresetManagerLocals
 {
@@ -59,7 +65,8 @@ namespace PresetManagerLocals
 		
 		AssetRegistryModule.Get().GetAssets(Filter, AssetData);
 
-		for (int i = 0; i < AssetData.Num(); i++) {
+		for (int i = 0; i < AssetData.Num(); i++)
+		{
 			ASSETCLASS* Object = Cast<ASSETCLASS>(AssetData[i].GetAsset());		
 			if (Object)
 			{				
@@ -82,6 +89,8 @@ namespace PresetManagerLocals
 		SLATE_BEGIN_ARGS(SCollectionTableRow) { }
 		    SLATE_ARGUMENT(TSharedPtr<SPresetManager::FPresetViewEntry>, ViewEntry)
 			SLATE_EVENT(FOnCollectionEnabledCheckboxChanged, OnCollectionEnabledCheckboxChanged)
+			SLATE_EVENT(FOnCollectionRenameStarted, OnCollectionRenameStarted)
+			SLATE_EVENT(FOnCollectionRenameEnded, OnCollectionRenameEnded)
 		SLATE_END_ARGS()
 
 		void Construct(const FArguments& InArgs, const TSharedRef<STableViewBase>& InOwnerTableView)
@@ -92,6 +101,8 @@ namespace PresetManagerLocals
 
 			ViewEntry = InArgs._ViewEntry;
 			OnCollectionEnabledCheckboxChanged = InArgs._OnCollectionEnabledCheckboxChanged;
+			OnCollectionRenameStarted = InArgs._OnCollectionRenameStarted;
+			OnCollectionRenameEnded = InArgs._OnCollectionRenameEnded;
 
 			STableRow<ItemType>::Construct(Args, InOwnerTableView);
 		}
@@ -117,9 +128,16 @@ namespace PresetManagerLocals
 					.VAlign(VAlign_Fill)
 					[
 						SAssignNew(EnabledWidget, SCheckBox)
-						.Visibility_Lambda([this]() {return ViewEntry->EntryType == EEntryType::Collection ? EVisibility::Visible : EVisibility::Collapsed; })
-						.IsChecked_Lambda([this]() {return ViewEntry->bEnabled ? ECheckBoxState::Checked : ECheckBoxState::Unchecked; })
-						.OnCheckStateChanged_Lambda([this](ECheckBoxState State) {
+						.Visibility_Lambda([this]()
+						{
+							return ViewEntry->EntryType == EEntryType::Collection ? EVisibility::Visible : EVisibility::Collapsed;
+						})
+						.IsChecked_Lambda([this]()
+						{
+							return ViewEntry->bEnabled ? ECheckBoxState::Checked : ECheckBoxState::Unchecked;
+						})
+						.OnCheckStateChanged_Lambda([this](ECheckBoxState State)
+						{
 							if (OnCollectionEnabledCheckboxChanged.IsBound())
 							{
 								OnCollectionEnabledCheckboxChanged.Execute(ViewEntry, State);
@@ -143,7 +161,10 @@ namespace PresetManagerLocals
 						.Padding(5.0f)
 						[
 							SNew(SImage)
-							.Visibility_Lambda([this](){return ViewEntry->EntryType == EEntryType::Tool ? EVisibility::Visible : EVisibility::Collapsed; })
+							.Visibility_Lambda([this]()
+							{
+								return ViewEntry->EntryType == EEntryType::Tool ? EVisibility::Visible : EVisibility::Collapsed;
+							})
 							.Image(&ViewEntry->EntryIcon)
 						]
 
@@ -151,18 +172,45 @@ namespace PresetManagerLocals
 						.FillWidth(1)
 						.Padding(5.0f)
 						[
-							SNew(STextBlock)
-							.Text(ViewEntry->EntryLabel)
-							.Font_Lambda([this]() {
-								if (ViewEntry->EntryType == EEntryType::Collection && ViewEntry->bIsDefaultCollection)
+							SNew(SHorizontalBox)
+
+							+ SHorizontalBox::Slot()
+							[
+								SNew(STextBlock)
+								.Text(ViewEntry->EntryLabel)
+								.Visibility_Lambda([this]()
 								{
-									return FCoreStyle::Get().GetWidgetStyle<FTextBlockStyle>("Text.Large").Font;
-								}
-								else
+									return ViewEntry->bIsRenaming ? EVisibility::Collapsed : EVisibility::Visible;
+								})							
+								.Font_Lambda([this]()
 								{
-									return FCoreStyle::Get().GetWidgetStyle<FTextBlockStyle>("NormalText").Font;
-								}
-							})
+									return FCoreStyle::Get().GetWidgetStyle<FTextBlockStyle>("NormalText").Font;							
+								})
+							]
+
+							+ SHorizontalBox::Slot()
+							[
+								SAssignNew(CollectionRenameBox, SEditableTextBox)
+								.Text(ViewEntry->EntryLabel)
+								
+								.Visibility_Lambda([this]()
+								{
+									if (!bIsEntryBeingRenamed && ViewEntry->bIsRenaming)
+									{
+										OnCollectionRenameStarted.ExecuteIfBound(ViewEntry, CollectionRenameBox);
+									}
+									bIsEntryBeingRenamed = ViewEntry->bIsRenaming;
+									return ViewEntry->bIsRenaming ? EVisibility::Visible : EVisibility::Collapsed;
+								})
+								.OnTextCommitted_Lambda([this](const FText & NewText, ETextCommit::Type CommitStatus)
+								{
+									OnCollectionRenameEnded.ExecuteIfBound(ViewEntry, NewText);
+								})
+								.Font_Lambda([this]()
+								{
+									return FCoreStyle::Get().GetWidgetStyle<FTextBlockStyle>("NormalText").Font;							
+								})
+							]
 						]
 
 
@@ -179,7 +227,10 @@ namespace PresetManagerLocals
 						.Padding(5.0f)
 						[
 							SNew(STextBlock)
-							.Text(TAttribute<FText>::CreateLambda([this]() { return FText::AsNumber(ViewEntry->Count); }))
+							.Text(TAttribute<FText>::CreateLambda([this]()
+							{
+								return FText::AsNumber(ViewEntry->Count);
+							}))
 						]
 				];
 			}
@@ -189,8 +240,13 @@ namespace PresetManagerLocals
 	private:
 		TSharedPtr<SPresetManager::FPresetViewEntry> ViewEntry;
 		TSharedPtr<SCheckBox> EnabledWidget;
+		TSharedPtr<SEditableTextBox> CollectionRenameBox;
+
+		bool bIsEntryBeingRenamed = false;
 
 		FOnCollectionEnabledCheckboxChanged OnCollectionEnabledCheckboxChanged;
+		FOnCollectionRenameStarted OnCollectionRenameStarted;
+		FOnCollectionRenameEnded OnCollectionRenameEnded;
 	};
 
 
@@ -237,8 +293,11 @@ namespace PresetManagerLocals
 					.Padding(5.0f)
 					[
 						SNew(SImage)
-						.Visibility_Lambda([this]() {return ViewEntry->EntryType == EEntryType::Tool ? EVisibility::Visible : EVisibility::Collapsed; })
-					.Image(&ViewEntry->EntryIcon)
+						.Visibility_Lambda([this]()
+						{
+							return ViewEntry->EntryType == EEntryType::Tool ? EVisibility::Visible : EVisibility::Collapsed;
+						})
+						.Image(&ViewEntry->EntryIcon)
 					]
 
 				+ SHorizontalBox::Slot()
@@ -262,7 +321,10 @@ namespace PresetManagerLocals
 					.Padding(5.0f)
 					[
 						SNew(STextBlock)
-						.Text(TAttribute<FText>::CreateLambda([this]() { return FText::AsNumber(ViewEntry->Count); }))
+						.Text(TAttribute<FText>::CreateLambda([this]()
+						{
+							return FText::AsNumber(ViewEntry->Count);
+						}))
 					]
 				];
 			}
@@ -289,7 +351,10 @@ namespace PresetManagerLocals
 					.Padding(FMargin(4.0f, 0.0f))
 					[
 						SNew(STextBlock)
-						.Text(TAttribute<FText>::CreateLambda([this]() { return FText::FromString(ViewEntry->PresetLabel); }))
+						.Text(TAttribute<FText>::CreateLambda([this]()
+						{
+							return FText::FromString(ViewEntry->PresetLabel);
+						}))
 					];
 			}
 			else if (InColumnName == "Tooltip")
@@ -300,7 +365,10 @@ namespace PresetManagerLocals
 					.Padding(FMargin(4.0f, 0.0f))
 					[
 						SNew(STextBlock)
-						.Text(TAttribute<FText>::CreateLambda([this]() { return FText::FromString(ViewEntry->PresetTooltip); }))
+						.Text(TAttribute<FText>::CreateLambda([this]()
+						{
+							return FText::FromString(ViewEntry->PresetTooltip);
+						}))
 					];
 			}
 			else if (InColumnName == "Tool")
@@ -318,8 +386,14 @@ namespace PresetManagerLocals
 			{
 				return SNew(SNegativeActionButton)
 					.Icon(FAppStyle::GetBrush("Icons.Delete"))						
-					.OnClicked_Lambda([this]() { OnPresetDeleted.ExecuteIfBound(ViewEntry); return FReply::Handled(); })
-					.Visibility_Lambda([this]() { return this->IsHovered() ? EVisibility::Visible : EVisibility::Hidden; });
+					.OnClicked_Lambda([this]()
+					{
+						OnPresetDeleted.ExecuteIfBound(ViewEntry); return FReply::Handled();
+					})
+					.Visibility_Lambda([this]()
+					{
+						return this->IsHovered() ? EVisibility::Visible : EVisibility::Hidden;
+					});
 			}
 
 			return SNullWidget::NullWidget;
@@ -343,6 +417,8 @@ BEGIN_SLATE_FUNCTION_BUILD_OPTIMIZATION
 void SPresetManager::Construct( const FArguments& InArgs )
 {
 	UPresetUserSettings::Initialize();
+	BindCommands();
+
 	UserSettings = UPresetUserSettings::Get();
 	if (UserSettings.IsValid())
 	{
@@ -353,36 +429,6 @@ void SPresetManager::Construct( const FArguments& InArgs )
 	[
 		SNew(SVerticalBox)
 
-		+ SVerticalBox::Slot()
-		.AutoHeight()
-		[
-			SNew(SHorizontalBox)
-
-			+ SHorizontalBox::Slot()
-			.AutoWidth()
-			.Padding(16.0f, 8.0f, 8.0f, 8.0f)
-			//.HAlign(HAlign_Fill)
-			//.SizeParam(FStretch(1.0f))
-			[
-				SAssignNew(AddUserPresetButton, SPositiveActionButton)
-				.Text(LOCTEXT( "AddUserPresetCollection", "Add User Preset Collection") )
-				.Icon(FAppStyle::GetBrush("Icons.PlusCircle"))
-				.OnClicked_Lambda([this]() { AddNewUserPresetCollection(); return FReply::Handled(); })
-			]
-
-			+ SHorizontalBox::Slot()
-			.AutoWidth()
-			.Padding(8.0f, 8.0f, 8.0f, 8.0f)
-			//.HAlign(HAlign_Fill)
-			//.SizeParam(FStretch(1.0f))
-			[
-				SAssignNew(DeleteUserPresetButton, SNegativeActionButton)
-				.Text( LOCTEXT( "DeleteUserPresetCollection", "Delete User Preset Collection"))
-				.Icon(FAppStyle::GetBrush("Icons.Delete"))
-				.OnClicked_Lambda([this]() { DeleteSelectedUserPresetCollection(); return FReply::Handled(); })
-			]
-
-		]
 		+ SVerticalBox::Slot()
 		.FillHeight(1.0)
 		//.AutoHeight()
@@ -395,33 +441,6 @@ void SPresetManager::Construct( const FArguments& InArgs )
 					.Resizable(false)
 						[
 							SNew(SVerticalBox)		
-								+ SVerticalBox::Slot()
-									.FillHeight(1.0f)
-								[
-			
-											SAssignNew(EditorPresetCollectionTreeView, STreeView<TSharedPtr<FPresetViewEntry> >)
-												.ItemHeight(32.0f)
-												.TreeItemsSource(&EditorCollectionsDataList)
-												.SelectionMode(ESelectionMode::Single)
-												.OnGenerateRow(this, &SPresetManager::HandleTreeGenerateRow)
-												.OnGetChildren(this, &SPresetManager::HandleTreeGetChildren)
-												.OnSelectionChanged(this, &SPresetManager::HandleEditorTreeSelectionChanged)
-												.HeaderRow
-												(
-													SNew(SHeaderRow)
-													.Visibility(EVisibility::Collapsed)
-
-													+ SHeaderRow::Column("Collection")
-													.FixedWidth(150.0f)
-													.HeaderContent()
-													[
-														SNew(STextBlock)
-														.Text(LOCTEXT("PresetManagerCollectionTitleHeader", "Collection"))
-													]
-
-
-												)											
-								]
 
 								+ SVerticalBox::Slot()
 								.AutoHeight()
@@ -441,7 +460,10 @@ void SPresetManager::Construct( const FArguments& InArgs )
 											.VAlign(VAlign_Center)
 											.HAlign(HAlign_Center)
 											.ClickMethod( EButtonClickMethod::MouseDown )
-											.OnClicked_Lambda([this]() { bAreUserCollectionsExpanded = !bAreUserCollectionsExpanded; return FReply::Handled(); })
+											.OnClicked_Lambda([this]()
+											{
+												bAreUserCollectionsExpanded = !bAreUserCollectionsExpanded; return FReply::Handled();
+											})
 											.ContentPadding(0.f)
 											.ForegroundColor( FSlateColor::UseForeground() )
 											.IsFocusable( false )
@@ -454,31 +476,79 @@ void SPresetManager::Construct( const FArguments& InArgs )
 									
 
 										+ SHorizontalBox::Slot()
-										.AutoWidth()
+										.FillWidth(1.0)
 										.Padding(5.0f, 5.0f, 5.0f, 5.0f)
-										.HAlign(HAlign_Center)
+										.HAlign(HAlign_Left)
 										.VAlign(VAlign_Center)
 										[
 											SNew(STextBlock)										
-											.Text(LOCTEXT("UserPresetLabels", "User Presets"))
+											.Text(LOCTEXT("UserPresetLabels", "User Preset Collections"))
 											.Font(FSlateFontInfo(FCoreStyle::GetDefaultFont(), 13, "Bold"))
 										]
+
+										+ SHorizontalBox::Slot()
+										.AutoWidth()
+										.HAlign(HAlign_Right)
+										.VAlign(VAlign_Center)
+										.Padding(16.0f, 8.0f, 8.0f, 8.0f)
+										[
+											SAssignNew(AddUserPresetButton, SPositiveActionButton)
+											.ToolTipText(LOCTEXT("AddUserPresetCollection", "Add User Preset Collection"))
+											.Icon(FAppStyle::GetBrush("Icons.Plus"))
+											.OnClicked_Lambda([this]()
+											{
+												AddNewUserPresetCollection(); return FReply::Handled();
+											})
+										]
 									]
 								]
 
+								+ SVerticalBox::Slot()
+								.AutoHeight()
+								[
+			
+											SAssignNew(EditorPresetCollectionTreeView, STreeView<TSharedPtr<FPresetViewEntry> >)
+												.TreeItemsSource(&EditorCollectionsDataList)
+												.SelectionMode(ESelectionMode::Single)
+												.OnGenerateRow(this, &SPresetManager::HandleTreeGenerateRow)
+												.OnGetChildren(this, &SPresetManager::HandleTreeGetChildren)
+												.OnSelectionChanged(this, &SPresetManager::HandleEditorTreeSelectionChanged)
+												.Visibility_Lambda([this]()
+												{
+													return bAreUserCollectionsExpanded ? EVisibility::Visible : EVisibility::Collapsed;
+												})
+												.HeaderRow
+												(
+													SNew(SHeaderRow)
+													.Visibility(EVisibility::Collapsed)
+
+													+ SHeaderRow::Column("Collection")
+													.FixedWidth(150.0f)
+													.HeaderContent()
+													[
+														SNew(STextBlock)
+														.Text(LOCTEXT("PresetManagerCollectionTitleHeader", "Collection"))
+													]
+
+
+												)											
+								]
+
 								+ SVerticalBox::Slot()
 									.FillHeight(1.0f)
-									//.AutoHeight()									
 								[
 			
 											SAssignNew(UserPresetCollectionTreeView, STreeView<TSharedPtr<FPresetViewEntry> >)
-												.ItemHeight(32.0f)
 												.TreeItemsSource(&UserCollectionsDataList)
 												.SelectionMode(ESelectionMode::Single)
 												.OnGenerateRow(this, &SPresetManager::HandleTreeGenerateRow)
 												.OnGetChildren(this, &SPresetManager::HandleTreeGetChildren)
 												.OnSelectionChanged(this, &SPresetManager::HandleUserTreeSelectionChanged)
-												.Visibility_Lambda([this]() {return bAreUserCollectionsExpanded ? EVisibility::Visible : EVisibility::Collapsed; })
+												.OnContextMenuOpening(this, &SPresetManager::OnGetCollectionContextMenuContent)
+												.Visibility_Lambda([this]()
+												{
+													return bAreUserCollectionsExpanded ? EVisibility::Visible : EVisibility::Collapsed;
+												})
 												.HeaderRow
 												(
 													SNew(SHeaderRow)
@@ -513,7 +583,10 @@ void SPresetManager::Construct( const FArguments& InArgs )
 											.VAlign(VAlign_Center)
 											.HAlign(HAlign_Center)
 											.ClickMethod( EButtonClickMethod::MouseDown )
-											.OnClicked_Lambda([this]() { bAreProjectCollectionsExpanded = !bAreProjectCollectionsExpanded; return FReply::Handled(); })
+											.OnClicked_Lambda([this]()
+											{
+												bAreProjectCollectionsExpanded = !bAreProjectCollectionsExpanded; return FReply::Handled();
+											})
 											.ContentPadding(0.f)
 											.ForegroundColor( FSlateColor::UseForeground() )
 											.IsFocusable( false )
@@ -525,25 +598,46 @@ void SPresetManager::Construct( const FArguments& InArgs )
 										]
 
 										+ SHorizontalBox::Slot()
-										.AutoWidth()
+										.FillWidth(1.0)
 										.Padding(5.0f, 5.0f, 5.0f, 5.0f)
-										.HAlign(HAlign_Center)
+										.HAlign(HAlign_Left)
 										.VAlign(VAlign_Center)
 										[
 											SNew(STextBlock)
-											.Text(LOCTEXT("ProjectPresetLabels", "Project Presets"))
+											.Text(LOCTEXT("ProjectPresetLabels", "Project Preset Collections"))
 											.Font(FSlateFontInfo(FCoreStyle::GetDefaultFont(), 13, "Bold"))
 										]
+
+										+ SHorizontalBox::Slot()
+										.AutoWidth()
+										.HAlign(HAlign_Right)
+										.VAlign(VAlign_Center)
+										.Padding(16.0f, 8.0f, 8.0f, 8.0f)
+										[
+											SNew(SSimpleButton)
+											.ToolTipText(LOCTEXT("OpenProjectSettingsPresets", "Open Project Settings for Presets"))
+											.Icon(FAppStyle::GetBrush("Icons.Settings"))
+											.OnClicked_Lambda([this]() 
+											{ 
+												if (ISettingsModule* SettingsModule = FModuleManager::GetModulePtr<ISettingsModule>("Settings"))
+												{
+													SettingsModule->ShowViewer("Project", "Plugins", "Interactive Tool Presets");
+												}
+												return FReply::Handled();
+											})
+										]
 									]
 								]
 
 								+ SVerticalBox::Slot()
 									.FillHeight(1.0f)
-									//.AutoHeight()
 								[
 										
 											SNew(SVerticalBox)
-											.Visibility_Lambda([this]() {return bAreProjectCollectionsExpanded ? EVisibility::Visible : EVisibility::Collapsed; })
+											.Visibility_Lambda([this]()
+											{
+												return bAreProjectCollectionsExpanded ? EVisibility::Visible : EVisibility::Collapsed;
+											})
 
 											+ SVerticalBox::Slot()
 											.FillHeight(1.0f)
@@ -557,6 +651,7 @@ void SPresetManager::Construct( const FArguments& InArgs )
 												.OnGenerateRow(this, &SPresetManager::HandleTreeGenerateRow)
 												.OnGetChildren(this, &SPresetManager::HandleTreeGetChildren)
 												.OnSelectionChanged(this, &SPresetManager::HandleTreeSelectionChanged)
+												.OnContextMenuOpening(this, &SPresetManager::OnGetCollectionContextMenuContent)
 												.HeaderRow
 												(
 													SNew(SHeaderRow)
@@ -577,10 +672,16 @@ void SPresetManager::Construct( const FArguments& InArgs )
 											+ SVerticalBox::Slot()
 											.AutoHeight()
 											.HAlign(HAlign_Center)
+											.Padding(5.0f)
 											[
 												SNew(STextBlock)
-												.Visibility_Lambda([this]() { return ProjectPresetCollectionsVisibility() == EVisibility::Visible ? EVisibility::Collapsed : EVisibility::Visible; })
-												.Text(LOCTEXT("ProjectPresetsNotLoadedLabel", "(Load Project Preset Collections in Project Settings)"))
+												.WrapTextAt(150.0f)												
+												.Visibility_Lambda([this]()
+												{
+													return ProjectPresetCollectionsVisibility() == EVisibility::Visible ? EVisibility::Collapsed : EVisibility::Visible;
+												})
+												.Text(LOCTEXT("ProjectPresetsNotLoadedLabel", "Manage Project Preset Collections in Project Settings"))
+												.Justification(ETextJustify::Center)
 												.Font(FAppStyle::GetFontStyle("NormalFontItalic"))
 											]
 
@@ -594,8 +695,13 @@ void SPresetManager::Construct( const FArguments& InArgs )
 							SNew(SVerticalBox)
 
 							+ SVerticalBox::Slot()
-							.FillHeight(1.0f)
+							.FillHeight(1.0)
 							[
+								SNew(SOverlay)
+								+ SOverlay::Slot()
+								.ZOrder(1)
+							[
+
 
 								SAssignNew(PresetListView, SListView<TSharedPtr<FPresetViewEntry>>)
 								.ListItemsSource(&PresetDataList)
@@ -603,6 +709,7 @@ void SPresetManager::Construct( const FArguments& InArgs )
 								.SelectionMode(ESelectionMode::SingleToggle)
 								.OnGenerateRow(this, &SPresetManager::HandleListGenerateRow)		
 								.OnSelectionChanged(this, &SPresetManager::HandleListSelectionChanged)
+								.OnContextMenuOpening(this, &SPresetManager::OnGetPresetContextMenuContent)
 								.HeaderRow
 								(
 									SNew(SHeaderRow)
@@ -641,19 +748,24 @@ void SPresetManager::Construct( const FArguments& InArgs )
 										SNew(STextBlock)
 										.Text(LOCTEXT("PresetManagerPresetTooltipHeader", "Tooltip"))
 									]
-
-									+ SHeaderRow::Column("Delete")
-									.FixedWidth(60.0f)									
-									.HAlignHeader(EHorizontalAlignment::HAlign_Center)
-									.VAlignHeader(EVerticalAlignment::VAlign_Center)
-									.HAlignCell(EHorizontalAlignment::HAlign_Center)
-									.HeaderContent()
-									[
-										SNew(SImage)
-										.Image(FAppStyle::GetBrush("Icons.Delete"))
-									]
-
 								)
+							]
+							
+							+ SOverlay::Slot()
+							.HAlign(HAlign_Center)
+							.VAlign(VAlign_Center)
+							.ZOrder(2)
+								[
+									SNew(STextBlock)
+									.WrapTextAt(150.0f)
+									.Visibility_Lambda([this]()
+									{
+										return bHasPresetsInCollection ? EVisibility::Collapsed : EVisibility::Visible;
+									})
+									.Text(LOCTEXT("NoPresetsAvailableLabel", "Add New Presets from any Modeling Tool"))
+									.Justification(ETextJustify::Center)
+									.Font(FAppStyle::GetFontStyle("NormalFontItalic"))
+								]
 						]
 
 						+ SVerticalBox::Slot()
@@ -661,7 +773,6 @@ void SPresetManager::Construct( const FArguments& InArgs )
 						.Padding(5.0f)
 						[
 							SNew(SVerticalBox)
-							.Visibility(this, &SPresetManager::EditAreaVisibility)
 							+ SVerticalBox::Slot()
 							.AutoHeight()
 							.Padding(5.0f)
@@ -680,22 +791,26 @@ void SPresetManager::Construct( const FArguments& InArgs )
 								.Padding(5.0f)
 								[
 									SNew(SEditableTextBox)
+									.IsEnabled(this, &SPresetManager::EditAreaEnabled)
 									.OverflowPolicy(ETextOverflowPolicy::Ellipsis)
-									.Text_Lambda([this]() {
+									.Text_Lambda([this]() 
+									{
 										if (ActivePresetToEdit)
 										{
 											return FText::FromString(ActivePresetToEdit->PresetLabel);
 										}
 										return FText::GetEmpty();
 									})
-									.OnTextChanged_Lambda([this](const FText& NewText) {
+									.OnTextChanged_Lambda([this](const FText& NewText)
+									{
 										if (ActivePresetToEdit)
 										{
 											// Cap the number of characters sent out of the text box, so we don't overflow menus and tooltips
 											ActivePresetToEdit->PresetLabel = NewText.ToString().Left(255);
 										}
 									})
-									.OnTextCommitted_Lambda([this](const FText& NewText, ETextCommit::Type CommitStatus){
+									.OnTextCommitted_Lambda([this](const FText& NewText, ETextCommit::Type CommitStatus)
+									{
 										if (ActivePresetToEdit)
 										{
 											// Cap the number of characters sent out of the text box, so we don't overflow menus and tooltips
@@ -726,26 +841,29 @@ void SPresetManager::Construct( const FArguments& InArgs )
 									.MaxDesiredHeight(44.0f)
 									[
 										SNew(SMultiLineEditableTextBox)		
-										
+										.IsEnabled(this, &SPresetManager::EditAreaEnabled)
 										.AllowMultiLine(false)
 										.AutoWrapText(true)
 										.WrappingPolicy(ETextWrappingPolicy::DefaultWrapping)
 										.OverflowPolicy(ETextOverflowPolicy::Ellipsis)
-										.Text_Lambda([this]() {
+										.Text_Lambda([this]()
+										{
 											if (ActivePresetToEdit)
 											{
 												return FText::FromString(ActivePresetToEdit->PresetTooltip);
 											}
 											return FText::GetEmpty();
 										})
-										.OnTextChanged_Lambda([this](const FText& NewText) {
+										.OnTextChanged_Lambda([this](const FText& NewText) 
+										{
 											if (ActivePresetToEdit)
 											{
 												// Cap the number of characters sent out of the text box, so we don't overflow menus and tooltips
 												ActivePresetToEdit->PresetTooltip = NewText.ToString().Left(2048);
 											}
 										})
-										.OnTextCommitted_Lambda([this](const FText& NewText, ETextCommit::Type CommitStatus){
+										.OnTextCommitted_Lambda([this](const FText& NewText, ETextCommit::Type CommitStatus)
+										{
 											if (ActivePresetToEdit)
 											{
 												// Cap the number of characters sent out of the text box, so we don't overflow menus and tooltips
@@ -760,7 +878,6 @@ void SPresetManager::Construct( const FArguments& InArgs )
 				]
 		];
 
-		DeleteUserPresetButton->SetEnabled(false);
 		RegeneratePresetTrees();
 		if (UserCollectionsDataList.Num() == 0)
 		{
@@ -845,7 +962,10 @@ void SPresetManager::RegeneratePresetTrees()
 
 		if (AssetList)
 		{
-			AssetList->RemoveAll([](const FSoftObjectPath& Path) { return !Path.IsAsset(); });
+			AssetList->RemoveAll([](const FSoftObjectPath& Path)
+			{
+				return !Path.IsAsset();
+			});
 
 			for (const FSoftObjectPath& Path : *AssetList)
 			{
@@ -931,7 +1051,9 @@ TSharedRef<ITableRow> SPresetManager::HandleTreeGenerateRow(TSharedPtr<FPresetVi
 {
 	return SNew(PresetManagerLocals::SCollectionTableRow< TSharedPtr<FPresetViewEntry> >, OwnerTable)
 		.ViewEntry(TreeEntry)
-		.OnCollectionEnabledCheckboxChanged(this, &SPresetManager::SetCollectionEnabled);
+		.OnCollectionEnabledCheckboxChanged(this, &SPresetManager::SetCollectionEnabled)
+		.OnCollectionRenameStarted(this, &SPresetManager::CollectionRenameStarted)
+		.OnCollectionRenameEnded(this, &SPresetManager::CollectionRenameEnded);
 }
 
 void SPresetManager::HandleTreeGetChildren(TSharedPtr<FPresetViewEntry> TreeEntry, TArray< TSharedPtr<FPresetViewEntry> >& ChildrenOut)
@@ -945,6 +1067,7 @@ void SPresetManager::GeneratePresetList(TSharedPtr<FPresetViewEntry> TreeEntry)
 	PresetListView->RequestListRefresh();
 	bHasActiveCollection = false;
 	ActivePresetToEdit = nullptr;
+	bHasPresetsInCollection = false;
 
 	if (!TreeEntry)
 	{
@@ -972,6 +1095,7 @@ void SPresetManager::GeneratePresetList(TSharedPtr<FPresetViewEntry> TreeEntry)
 					{
 						if (ToolNameIter.Value().NamedPresets[PresetIndex].IsValid())
 						{
+							bHasPresetsInCollection = true;
 							PresetDataList.Add(MakeShared<FPresetViewEntry>(
 								ToolNameIter.Key(),
 								PresetIndex,
@@ -1002,6 +1126,7 @@ void SPresetManager::GeneratePresetList(TSharedPtr<FPresetViewEntry> TreeEntry)
 				{
 					if (ToolData->NamedPresets[PresetIndex].IsValid())
 					{
+						bHasPresetsInCollection = true;
 						PresetDataList.Add(MakeShared<FPresetViewEntry>(
 							TreeEntry->ToolName,
 							PresetIndex,
@@ -1022,39 +1147,55 @@ void SPresetManager::GeneratePresetList(TSharedPtr<FPresetViewEntry> TreeEntry)
 
 void SPresetManager::HandleEditorTreeSelectionChanged(TSharedPtr<FPresetViewEntry> TreeEntry, ESelectInfo::Type SelectInfo)
 {
-	if (SelectInfo != ESelectInfo::Direct) {
+	for (TSharedPtr<FPresetViewEntry> Entry : UserPresetCollectionTreeView->GetRootItems())
+	{
+		Entry->bIsRenaming = false;
+	}
+
+	if (SelectInfo != ESelectInfo::Direct)
+	{
 		UserPresetCollectionTreeView->ClearSelection();
 		ProjectPresetCollectionTreeView->ClearSelection();
 		GeneratePresetList(TreeEntry);
+
+		LastFocusedList = EditorPresetCollectionTreeView;
 	}
 }
 
 void SPresetManager::HandleTreeSelectionChanged(TSharedPtr<FPresetViewEntry> TreeEntry, ESelectInfo::Type SelectInfo)
 {
-	if (SelectInfo != ESelectInfo::Direct) {
+	for (TSharedPtr<FPresetViewEntry> Entry : UserPresetCollectionTreeView->GetRootItems())
+	{
+		Entry->bIsRenaming = false;
+	}
+
+	if (SelectInfo != ESelectInfo::Direct)
+	{
 		UserPresetCollectionTreeView->ClearSelection();
 		EditorPresetCollectionTreeView->ClearSelection();
 		GeneratePresetList(TreeEntry);
+
+		LastFocusedList = ProjectPresetCollectionTreeView;
+
 	}
 }
 
 void SPresetManager::HandleUserTreeSelectionChanged(TSharedPtr<FPresetViewEntry> TreeEntry, ESelectInfo::Type SelectInfo)
 {
-	if (SelectInfo != ESelectInfo::Direct) {
+	for (TSharedPtr<FPresetViewEntry> Entry : UserPresetCollectionTreeView->GetRootItems())
+	{
+		Entry->bIsRenaming = false;
+	}
+
+	if (SelectInfo != ESelectInfo::Direct)
+	{
 		ProjectPresetCollectionTreeView->ClearSelection();
 		EditorPresetCollectionTreeView->ClearSelection();
 		GeneratePresetList(TreeEntry);
-	}
 
-	DeleteUserPresetButton->SetEnabled(false);
-	if (TreeEntry)
-	{
-		if (TreeEntry->EntryType == FPresetViewEntry::EEntryType::Collection && !TreeEntry->bIsDefaultCollection)
-		{
-			DeleteUserPresetButton->SetEnabled(true);
-		}
-	}
+		LastFocusedList = UserPresetCollectionTreeView;
 
+	}
 }
 
 TSharedRef<ITableRow> SPresetManager::HandleListGenerateRow(TSharedPtr<FPresetViewEntry> TreeEntry, const TSharedRef<STableViewBase>& OwnerTable)
@@ -1066,6 +1207,12 @@ TSharedRef<ITableRow> SPresetManager::HandleListGenerateRow(TSharedPtr<FPresetVi
 
 void SPresetManager::HandleListSelectionChanged(TSharedPtr<FPresetViewEntry> TreeEntry, ESelectInfo::Type SelectInfo)
 {
+	if (SelectInfo != ESelectInfo::Direct)
+	{
+		LastFocusedList = PresetListView;
+	}
+
+
 	if (TreeEntry)
 	{
 		ActivePresetToEdit = TreeEntry;
@@ -1081,9 +1228,9 @@ void SPresetManager::HandleListSelectionChanged(TSharedPtr<FPresetViewEntry> Tre
 	}
 }
 
-EVisibility SPresetManager::EditAreaVisibility() const
+bool SPresetManager::EditAreaEnabled() const
 {
-	return ActivePresetToEdit.IsValid() ? EVisibility::Visible : EVisibility::Collapsed;
+	return ActivePresetToEdit.IsValid();
 }
 
 EVisibility SPresetManager::ProjectPresetCollectionsVisibility() const
@@ -1117,6 +1264,52 @@ void SPresetManager::SetCollectionEnabled(TSharedPtr<FPresetViewEntry> TreeEntry
 	}
 }
 
+void SPresetManager::CollectionRenameStarted(TSharedPtr<FPresetViewEntry> TreeEntry, TSharedPtr<SEditableTextBox> RenameWidget)
+{
+	// TODO: Figure out why this crashes
+	//FSlateApplication::Get().SetKeyboardFocus(RenameWidget, EFocusCause::SetDirectly);
+}
+
+void SPresetManager::CollectionRenameEnded(TSharedPtr<FPresetViewEntry> TreeEntry, const FText& NewText)
+{
+	FAssetRegistryModule& AssetRegistryModule = FModuleManager::LoadModuleChecked<FAssetRegistryModule>(TEXT("AssetRegistry"));
+
+	TreeEntry->bIsRenaming = false;
+
+	FAssetData CollectionAsset;
+	if (AssetRegistryModule.Get().TryGetAssetByObjectPath(TreeEntry->CollectionPath, CollectionAsset) == UE::AssetRegistry::EExists::Exists)
+	{
+		UInteractiveToolsPresetCollectionAsset* CollectionObject = Cast<UInteractiveToolsPresetCollectionAsset>(CollectionAsset.GetAsset());
+		if (!CollectionObject)
+		{
+			return;
+		}
+
+		TArray< FAssetRenameData> RenameData;
+
+		FString NewPackageName, NewAssetName;
+		IAssetTools::Get().CreateUniqueAssetName(NewText.ToString(), "", NewPackageName, NewAssetName);
+
+		RenameData.SetNum(1);		
+		RenameData[0].Asset = CollectionObject;
+		RenameData[0].NewName = NewAssetName;
+		RenameData[0].NewPackagePath = CollectionAsset.PackagePath.ToString();
+		if(IAssetTools::Get().RenameAssets(RenameData))
+		{
+			if (UserSettings->EnabledPresetCollections.Contains(TreeEntry->CollectionPath))
+			{
+				UserSettings->EnabledPresetCollections.Remove(TreeEntry->CollectionPath);
+				UserSettings->EnabledPresetCollections.Add(CollectionObject->GetPathName());
+				UserSettings->SaveEditorConfig();
+			}
+
+			CollectionObject->CollectionLabel = NewText;
+			CollectionObject->MarkPackageDirty();	
+		}
+	}
+}
+
+
 void SPresetManager::DeletePresetFromCollection(TSharedPtr< FPresetViewEntry > Entry)
 {
 	UInteractiveToolsPresetCollectionAsset* PresetCollection = GetCollectionFromEntry(Entry);
@@ -1177,6 +1370,8 @@ void SPresetManager::DeleteSelectedUserPresetCollection()
 			AssetData.Add(CollectionAsset);
 			ObjectTools::DeleteAssets(AssetData, true);
 		}
+
+		GeneratePresetList(nullptr);
 	}
 }
 
@@ -1205,6 +1400,9 @@ void SPresetManager::AddNewUserPresetCollection()
 
 	// Inform asset registry
 	AssetRegistry.AssetCreated(NewObject);
+
+	// Since we're adding a new entry, open this tree view up again.
+	bAreUserCollectionsExpanded = true;
 }
 
 const FSlateBrush* SPresetManager::GetProjectCollectionsExpanderImage() const
@@ -1282,5 +1480,110 @@ void SPresetManager::SaveIfDefaultCollection(TSharedPtr<FPresetViewEntry> Entry)
 	}
 }
 
+TSharedPtr<SWidget> SPresetManager::OnGetPresetContextMenuContent() const
+{
+	const bool bShouldCloseWindowAfterMenuSelection = true;
+	FMenuBuilder MenuBuilder(bShouldCloseWindowAfterMenuSelection, UICommandList);
+
+	MenuBuilder.BeginSection("PresetManagerPresetAction", LOCTEXT("PresetAction", "Preset Actions"));
+
+	MenuBuilder.AddMenuEntry(FGenericCommands::Get().Delete, NAME_None, LOCTEXT("DeletePresetLabel", "Delete Preset"), LOCTEXT("DeletePresetToolTip", "Delete the selected preset"));
+
+	MenuBuilder.EndSection();
+
+	return MenuBuilder.MakeWidget();
+}
+
+TSharedPtr<SWidget> SPresetManager::OnGetCollectionContextMenuContent() const
+{
+	const bool bShouldCloseWindowAfterMenuSelection = true;
+	FMenuBuilder MenuBuilder(bShouldCloseWindowAfterMenuSelection, UICommandList);
+
+	MenuBuilder.BeginSection("PresetManagerCollectionAction", LOCTEXT("CollectionAction", "Preset Collection Actions"));
+
+	MenuBuilder.AddMenuEntry(FGenericCommands::Get().Delete, NAME_None, LOCTEXT("DeleteCollectionLabel", "Delete Collection"), LOCTEXT("DeleteCollectionToolTip", "Delete the selected collection"));
+	MenuBuilder.AddMenuEntry(FGenericCommands::Get().Rename, NAME_None, LOCTEXT("RenameCollectionLabel", "Rename Collection"), LOCTEXT("RenameCollectionToolTip", "Rename the selected collection"));
+
+	MenuBuilder.EndSection();
+
+	return MenuBuilder.MakeWidget();
+}
+
+void SPresetManager::BindCommands()
+{
+	// This should not be called twice on the same instance
+	check(!UICommandList.IsValid());
+
+	UICommandList = MakeShareable(new FUICommandList);
+
+	FUICommandList& CommandList = *UICommandList;
+
+	// ...and bind them all
+
+	CommandList.MapAction(
+		FGenericCommands::Get().Delete,
+		FExecuteAction::CreateSP(this, &SPresetManager::OnDeleteClicked),
+		FCanExecuteAction::CreateSP(this, &SPresetManager::CanDelete));
+
+	CommandList.MapAction(
+		FGenericCommands::Get().Rename,
+		FExecuteAction::CreateSP(this, &SPresetManager::OnRenameClicked),
+		FCanExecuteAction::CreateSP(this, &SPresetManager::CanRename));	
+}
+
+
+void SPresetManager::OnDeleteClicked()
+{
+	if (UserPresetCollectionTreeView == LastFocusedList)
+	{
+		DeleteSelectedUserPresetCollection();
+	}
+
+	if (PresetListView == LastFocusedList)
+	{
+		for (TSharedPtr<FPresetViewEntry> Entry : PresetListView->GetSelectedItems())
+		{
+			DeletePresetFromCollection(Entry);
+		}
+	}
+}
+
+bool SPresetManager::CanDelete()
+{
+	bool bIsListValid = LastFocusedList.IsValid() && (LastFocusedList == UserPresetCollectionTreeView || LastFocusedList == PresetListView);
+	bool bIsSelectionValid = false;
+	if (bIsListValid)
+	{
+		bIsSelectionValid = LastFocusedList.Pin()->GetNumItemsSelected() == 1 &&
+			                (LastFocusedList.Pin()->GetSelectedItems()[0]->EntryType == FPresetViewEntry::EEntryType::Collection ||
+							 LastFocusedList.Pin()->GetSelectedItems()[0]->EntryType == FPresetViewEntry::EEntryType::Preset);
+	}
+	return bIsSelectionValid;
+}
+
+void SPresetManager::OnRenameClicked()
+{
+	for (TSharedPtr<FPresetViewEntry> Entry : UserPresetCollectionTreeView->GetRootItems())
+	{
+		Entry->bIsRenaming = false;
+	}
+
+	for (TSharedPtr<FPresetViewEntry> Entry : UserPresetCollectionTreeView->GetSelectedItems())
+	{
+		Entry->bIsRenaming = true;
+	}
+}
+
+bool SPresetManager::CanRename()
+{
+	bool bIsListValid = LastFocusedList.IsValid() && LastFocusedList == UserPresetCollectionTreeView;
+	bool bIsSelectionValid = false;
+	if (bIsListValid)
+	{
+		bIsSelectionValid = LastFocusedList.Pin()->GetNumItemsSelected() == 1 &&
+			                LastFocusedList.Pin()->GetSelectedItems()[0]->EntryType == FPresetViewEntry::EEntryType::Collection;
+	}
+	return bIsSelectionValid;
+}
 
 #undef LOCTEXT_NAMESPACE
\ No newline at end of file
diff --git a/Engine/Plugins/Experimental/ToolPresets/Source/PresetEditor/Private/Widgets/SPresetManager.h b/Engine/Plugins/Experimental/ToolPresets/Source/PresetEditor/Private/Widgets/SPresetManager.h
--- a/Engine/Plugins/Experimental/ToolPresets/Source/PresetEditor/Private/Widgets/SPresetManager.h
+++ b/Engine/Plugins/Experimental/ToolPresets/Source/PresetEditor/Private/Widgets/SPresetManager.h
@@ -20,7 +20,7 @@ class SEditableTextBox;
 template<class ItemType> class STreeView;
 template<class ItemType> class SListView;
 class UInteractiveToolsPresetCollectionAsset;
-
+class FUICommandList;
 
 /**
  * Implements the preset manager panel.
@@ -56,6 +56,7 @@ public:
 		bool bEnabled = false;
 		FSoftObjectPath CollectionPath;
 		bool bIsDefaultCollection = false;
+		bool bIsRenaming = false;
 		FText EntryLabel;
 		FSlateBrush EntryIcon;
 		int32 Count = 0;
@@ -173,6 +174,8 @@ protected:
 
 private:
 
+	void BindCommands();
+
 	void RegeneratePresetTrees();
 
 	int32 GetTotalPresetCount() const;
@@ -184,14 +187,20 @@ private:
 	void HandleEditorTreeSelectionChanged(TSharedPtr<FPresetViewEntry> TreeEntry, ESelectInfo::Type);
 	EVisibility ProjectPresetCollectionsVisibility() const;
 
+	TSharedPtr<SWidget> OnGetPresetContextMenuContent() const;
+	TSharedPtr<SWidget> OnGetCollectionContextMenuContent() const;
+
 	void GeneratePresetList(TSharedPtr<FPresetViewEntry> TreeEntry);
 	TSharedRef<ITableRow> HandleListGenerateRow(TSharedPtr<FPresetViewEntry> TreeEntry, const TSharedRef<STableViewBase>& OwnerTable);
     void HandleListSelectionChanged(TSharedPtr<FPresetViewEntry> TreeEntry, ESelectInfo::Type SelectInfo);
 
-	EVisibility EditAreaVisibility() const;
+	bool EditAreaEnabled() const;
 	void SetCollectionEnabled(TSharedPtr<FPresetViewEntry> TreeEntry, ECheckBoxState State);
 	void DeletePresetFromCollection(TSharedPtr< FPresetViewEntry > Entry);
 
+	void CollectionRenameStarted(TSharedPtr<FPresetViewEntry> TreeEntry, TSharedPtr<SEditableTextBox> RenameWidget);
+	void CollectionRenameEnded(TSharedPtr<FPresetViewEntry> TreeEntry, const FText& NewText);
+
 	void DeleteSelectedUserPresetCollection();
 	void AddNewUserPresetCollection();
 
@@ -205,10 +214,20 @@ private:
 	UInteractiveToolsPresetCollectionAsset* GetCollectionFromEntry(TSharedPtr<FPresetViewEntry> Entry);
 	void SaveIfDefaultCollection(TSharedPtr<FPresetViewEntry> Entry);
 
+	void OnDeleteClicked();
+	bool CanDelete();
+
+	void OnRenameClicked();
+	bool CanRename();
+
 private:
+	
+	TSharedPtr<FUICommandList> UICommandList;
 
 	TWeakObjectPtr<UPresetUserSettings> UserSettings;
 
+	TWeakPtr< SListView<TSharedPtr<FPresetViewEntry> > >  LastFocusedList;
+
 	bool bAreProjectCollectionsExpanded = true;
 	TSharedPtr<SButton> ProjectCollectionsExpander;
 	TArray< TSharedPtr< FPresetViewEntry > > ProjectCollectionsDataList;
@@ -228,6 +247,7 @@ private:
 
 	int32 TotalPresetCount;
 	bool bHasActiveCollection;
+	bool bHasPresetsInCollection;
 	FText ActiveCollectionLabel;
 	bool bIsActiveCollectionEnabled;
 
diff --git a/Engine/Source/Editor/ContentBrowser/Private/AssetViewWidgets.cpp b/Engine/Source/Editor/ContentBrowser/Private/AssetViewWidgets.cpp
--- a/Engine/Source/Editor/ContentBrowser/Private/AssetViewWidgets.cpp
+++ b/Engine/Source/Editor/ContentBrowser/Private/AssetViewWidgets.cpp
@@ -1911,6 +1911,7 @@ void SAssetTileItem::Construct( const FArguments& InArgs )
 					]
 
 					+SVerticalBox::Slot()
+					.AutoHeight()
 					[
 						SNew(SBorder)
 						.Padding(FMargin(2.0f, 3.0f))
@@ -1920,6 +1921,7 @@ void SAssetTileItem::Construct( const FArguments& InArgs )
 							+ SVerticalBox::Slot()
 							.Padding(2.0f,2.0f,0.0f,0.0f)
 							.VAlign(VAlign_Top)
+						    .AutoHeight()
 							.HAlign(IsFolder() ? HAlign_Center : HAlign_Left)
 							[
 								SNew(SBox)
