[ToolPresets] Renamed the modules to match the plugin name

diff --git a/Engine/Plugins/Editor/ModelingToolsEditorMode/Source/ModelingToolsEditorMode/ModelingToolsEditorMode.Build.cs b/Engine/Plugins/Editor/ModelingToolsEditorMode/Source/ModelingToolsEditorMode/ModelingToolsEditorMode.Build.cs
--- a/Engine/Plugins/Editor/ModelingToolsEditorMode/Source/ModelingToolsEditorMode/ModelingToolsEditorMode.Build.cs
+++ b/Engine/Plugins/Editor/ModelingToolsEditorMode/Source/ModelingToolsEditorMode/ModelingToolsEditorMode.Build.cs
@@ -51,8 +51,8 @@ public class ModelingToolsEditorMode : ModuleRules
 				"PropertyEditor",
 				"ToolMenus",
 				"EditorConfig",
-				"PresetAsset",
-				"PresetEditor",
+				"ToolPresetAsset",
+				"ToolPresetEditor",
 				"EditorConfig",
 				// ... add private dependencies that you statically link with here ...	
 			}
diff --git a/Engine/Plugins/Editor/ModelingToolsEditorMode/Source/ModelingToolsEditorMode/Private/ModelingToolsEditorModeToolkit.cpp b/Engine/Plugins/Editor/ModelingToolsEditorMode/Source/ModelingToolsEditorMode/Private/ModelingToolsEditorModeToolkit.cpp
--- a/Engine/Plugins/Editor/ModelingToolsEditorMode/Source/ModelingToolsEditorMode/Private/ModelingToolsEditorModeToolkit.cpp
+++ b/Engine/Plugins/Editor/ModelingToolsEditorMode/Source/ModelingToolsEditorMode/Private/ModelingToolsEditorModeToolkit.cpp
@@ -48,7 +48,7 @@
 #include "Framework/Notifications/NotificationManager.h"
 
 // for presets
-#include "PresetAsset.h"
+#include "ToolPresetAsset.h"
 #include "PropertyCustomizationHelpers.h"
 
 #include "LevelEditorViewport.h"
@@ -58,12 +58,12 @@
 #include "ToolkitBuilder.h"
 #include "Dialogs/Dialogs.h"
 #include "Widgets/Input/SComboButton.h"
-#include "IPresetEditorModule.h"
-#include "PresetSettings.h"
+#include "IToolPresetEditorModule.h"
+#include "ToolPresetSettings.h"
 #include "AssetRegistry/AssetRegistryModule.h"
 #include "ModelingWidgets/SToolInputAssetComboPanel.h"
 #include "Fonts/SlateFontInfo.h"
-#include "PresetAssetSubsystem.h"
+#include "ToolPresetAssetSubsystem.h"
 
 
 #define LOCTEXT_NAMESPACE "FModelingToolsEditorModeToolkit"
@@ -92,7 +92,7 @@ namespace FModelingToolsEditorModeToolkitLocals
 	{
 		UInteractiveToolsPresetCollectionAsset* Preset = nullptr;
 
-		UPresetAssetSubsystem* PresetAssetSubsystem = GEditor->GetEditorSubsystem<UPresetAssetSubsystem>();
+		UToolPresetAssetSubsystem* PresetAssetSubsystem = GEditor->GetEditorSubsystem<UToolPresetAssetSubsystem>();
 
 		if (PresetPath.IsNull() && ensure(PresetAssetSubsystem))
 		{
@@ -116,7 +116,7 @@ namespace FModelingToolsEditorModeToolkitLocals
 		UInteractiveToolsPresetCollectionAsset* Preset = nullptr;		
 		UInteractiveTool* Tool = EdMode.GetToolManager()->GetActiveTool(EToolSide::Left);
 
-		UPresetAssetSubsystem* PresetAssetSubsystem = GEditor->GetEditorSubsystem<UPresetAssetSubsystem>();
+		UToolPresetAssetSubsystem* PresetAssetSubsystem = GEditor->GetEditorSubsystem<UToolPresetAssetSubsystem>();
 
 		if (PresetPath.IsNull() && ensure(PresetAssetSubsystem))
 		{
@@ -163,8 +163,8 @@ FModelingToolsEditorModeToolkit::FModelingToolsEditorModeToolkit()
 	UModelingModeEditableToolPaletteConfig::Initialize();
 	UModelingModeEditableToolPaletteConfig::Get()->LoadEditorConfig();
 
-	UPresetUserSettings::Initialize();
-	UPresetUserSettings::Get()->LoadEditorConfig();
+	UToolPresetUserSettings::Initialize();
+	UToolPresetUserSettings::Get()->LoadEditorConfig();
 
 	RecentPresetCollectionProvider = MakeShared< FRecentPresetCollectionProvider>();
 	CurrentPreset = MakeShared<FAssetData>();
@@ -1166,6 +1166,7 @@ TSharedRef<SWidget> FModelingToolsEditorModeToolkit::GetPresetCreateButtonConten
 
 		MenuBuilder.BeginSection(NoExtensionHook, LOCTEXT("ModelingPresetPanelHeaderManagePresets", "Manage Presets"));
 
+
 		FUIAction CreateNewPresetAction;
 		CreateNewPresetAction.ExecuteAction = FExecuteAction::CreateLambda(OpenNewPresetDialog);
 		MenuBuilder.AddMenuEntry(
@@ -1174,9 +1175,10 @@ TSharedRef<SWidget> FModelingToolsEditorModeToolkit::GetPresetCreateButtonConten
 			FSlateIcon( FAppStyle::Get().GetStyleSetName(), "Icons.Plus"),
 			CreateNewPresetAction);
 
+
 		FUIAction OpenPresetManangerAction;
 		OpenPresetManangerAction.ExecuteAction = FExecuteAction::CreateLambda([this]() {
-			IPresetEditorModule::Get().ExecuteOpenPresetEditor();
+			IToolPresetEditorModule::Get().ExecuteOpenPresetEditor();
 		});
 		MenuBuilder.AddMenuEntry(
 			FText(LOCTEXT("ModelingPresetPanelOpenPresetMananger", "Manage Presets...")),
@@ -1198,11 +1200,11 @@ void FModelingToolsEditorModeToolkit::ClearPresetComboList()
 
 void FModelingToolsEditorModeToolkit::RebuildPresetListForTool(bool bSettingsOpened)
 {	
-	TObjectPtr<UPresetUserSettings> UserSettings = UPresetUserSettings::Get();
+	TObjectPtr<UToolPresetUserSettings> UserSettings = UToolPresetUserSettings::Get();
 	UserSettings->LoadEditorConfig();
 
 	// We need to generate a combined list of Project Loaded and User available presets to intersect the enabled set against...
-	const UPresetProjectSettings* ProjectSettings = GetDefault<UPresetProjectSettings>();
+	const UToolPresetProjectSettings* ProjectSettings = GetDefault<UToolPresetProjectSettings>();
 	TSet<FSoftObjectPath> AllUserPresets;
 	
 	FAssetRegistryModule& AssetRegistryModule = FModuleManager::LoadModuleChecked<FAssetRegistryModule>("AssetRegistry");
@@ -1244,7 +1246,7 @@ void FModelingToolsEditorModeToolkit::RebuildPresetListForTool(bool bSettingsOpe
 				AvailablePresetsForTool.Reserve(Preset.PerToolPresets[Tool.GetClass()->GetName()].NamedPresets.Num());
 				for (int32 PresetIndex = 0; PresetIndex < Preset.PerToolPresets[Tool.GetClass()->GetName()].NamedPresets.Num(); ++PresetIndex)
 				{
-					const FInteractiveToolPresetDefintion& PresetDef = Preset.PerToolPresets[Tool.GetClass()->GetName()].NamedPresets[PresetIndex];
+					const FInteractiveToolPresetDefinition& PresetDef = Preset.PerToolPresets[Tool.GetClass()->GetName()].NamedPresets[PresetIndex];
 					if (!PresetDef.IsValid())
 					{
 						continue;
@@ -1321,7 +1323,7 @@ void FModelingToolsEditorModeToolkit::CreateNewPresetInCollection(const FString&
 		{
 			Preset.PerToolPresets.FindOrAdd(Tool.GetClass()->GetName()).ToolIcon = *ActiveToolIcon;
 		}
-		FInteractiveToolPresetDefintion PresetValuesToCreate;
+		FInteractiveToolPresetDefinition PresetValuesToCreate;
 		if (PresetLabel.IsEmpty())
 		{
 			PresetValuesToCreate.Label = FString::Printf(TEXT("Unnamed_Preset-%d"), Preset.PerToolPresets.FindOrAdd(Tool.GetClass()->GetName()).NamedPresets.Num()+1);
@@ -1350,7 +1352,7 @@ void FModelingToolsEditorModeToolkit::CreateNewPresetInCollection(const FString&
 		
 	});
 
-	UPresetAssetSubsystem* PresetAssetSubsystem = GEditor->GetEditorSubsystem<UPresetAssetSubsystem>();
+	UToolPresetAssetSubsystem* PresetAssetSubsystem = GEditor->GetEditorSubsystem<UToolPresetAssetSubsystem>();
 	if (CollectionPath.IsNull() && ensure(PresetAssetSubsystem))
 	{
 		ensure(PresetAssetSubsystem->SaveDefaultCollection());
@@ -1391,7 +1393,7 @@ void FModelingToolsEditorModeToolkit::UpdatePresetInCollection(const FToolPreset
 
 	RebuildPresetListForTool(false);
 
-	UPresetAssetSubsystem* PresetAssetSubsystem = GEditor->GetEditorSubsystem<UPresetAssetSubsystem>();
+	UToolPresetAssetSubsystem* PresetAssetSubsystem = GEditor->GetEditorSubsystem<UToolPresetAssetSubsystem>();
 	if (PresetToEditIn.PresetCollection.IsNull() && ensure(PresetAssetSubsystem))
 	{
 		ensure(PresetAssetSubsystem->SaveDefaultCollection());
diff --git a/Engine/Plugins/Experimental/ToolPresets/Config/BaseToolPresets.ini b/Engine/Plugins/Experimental/ToolPresets/Config/BaseToolPresets.ini
new file mode 100644
--- /dev/null
+++ b/Engine/Plugins/Experimental/ToolPresets/Config/BaseToolPresets.ini
@@ -0,0 +1,11 @@
+[CoreRedirects]
++PackageRedirects=(OldName="/Script/PresetEditor",NewName="/Script/ToolPresetEditor")
++ClassRedirects=(OldName="/Script/PresetEditor.PresetUserSettings",NewName="/Script/ToolPresetEditor.ToolPresetUserSettings")
++ClassRedirects=(OldName="/Script/PresetEditor.PresetProjectSettings",NewName="/Script/ToolPresetEditor.ToolPresetProjectSettings")
+
++PackageRedirects=(OldName="/Script/PresetAsset",NewName="/Script/ToolPresetAsset")
++StructRedirects=(OldName="/Script/PresetAsset.InteractiveToolPresetDefintion",NewName="/Script/ToolPresetAsset.InteractiveToolPresetDefinition")
++StructRedirects=(OldName="/Script/PresetAsset.InteractiveToolPresetStore",NewName="/Script/ToolPresetAsset.InteractiveToolPresetStore")
++ClassRedirects=(OldName="/Script/PresetAsset.InteractiveToolsPresetCollectionAsset",NewName="/Script/ToolPresetAsset.InteractiveToolsPresetCollectionAsset")
++ClassRedirects=(OldName="/Script/PresetAsset.InteractiveToolsPresetCollectionAssetFactory",NewName="/Script/ToolPresetAsset.InteractiveToolsPresetCollectionAssetFactory")
++ClassRedirects=(OldName="/Script/PresetAsset.PresetAssetSubsystem",NewName="/Script/ToolPresetAsset.ToolPresetAssetSubsystem")
diff --git a/Engine/Plugins/Experimental/ToolPresets/Source/PresetAsset/PresetAsset.Build.cs b/Engine/Plugins/Experimental/ToolPresets/Source/PresetAsset/PresetAsset.Build.cs
deleted file mode 100644
diff --git a/Engine/Plugins/Experimental/ToolPresets/Source/PresetAsset/Private/AssetDefinition_PresetAsset.cpp b/Engine/Plugins/Experimental/ToolPresets/Source/PresetAsset/Private/AssetDefinition_PresetAsset.cpp
deleted file mode 100644
diff --git a/Engine/Plugins/Experimental/ToolPresets/Source/PresetAsset/Private/AssetDefinition_PresetAsset.h b/Engine/Plugins/Experimental/ToolPresets/Source/PresetAsset/Private/AssetDefinition_PresetAsset.h
deleted file mode 100644
diff --git a/Engine/Plugins/Experimental/ToolPresets/Source/PresetAsset/Private/PresetAsset.cpp b/Engine/Plugins/Experimental/ToolPresets/Source/PresetAsset/Private/PresetAsset.cpp
deleted file mode 100644
diff --git a/Engine/Plugins/Experimental/ToolPresets/Source/PresetAsset/Private/PresetAssetModule.cpp b/Engine/Plugins/Experimental/ToolPresets/Source/PresetAsset/Private/PresetAssetModule.cpp
deleted file mode 100644
diff --git a/Engine/Plugins/Experimental/ToolPresets/Source/PresetAsset/Private/PresetAssetSubsystem.cpp b/Engine/Plugins/Experimental/ToolPresets/Source/PresetAsset/Private/PresetAssetSubsystem.cpp
deleted file mode 100644
diff --git a/Engine/Plugins/Experimental/ToolPresets/Source/PresetAsset/Public/PresetAsset.h b/Engine/Plugins/Experimental/ToolPresets/Source/PresetAsset/Public/PresetAsset.h
deleted file mode 100644
diff --git a/Engine/Plugins/Experimental/ToolPresets/Source/PresetAsset/Public/PresetAssetSubsystem.h b/Engine/Plugins/Experimental/ToolPresets/Source/PresetAsset/Public/PresetAssetSubsystem.h
deleted file mode 100644
diff --git a/Engine/Plugins/Experimental/ToolPresets/Source/PresetEditor/PresetEditor.Build.cs b/Engine/Plugins/Experimental/ToolPresets/Source/PresetEditor/PresetEditor.Build.cs
deleted file mode 100644
diff --git a/Engine/Plugins/Experimental/ToolPresets/Source/PresetEditor/Private/PresetEditorModule.cpp b/Engine/Plugins/Experimental/ToolPresets/Source/PresetEditor/Private/PresetEditorModule.cpp
deleted file mode 100644
diff --git a/Engine/Plugins/Experimental/ToolPresets/Source/PresetEditor/Private/PresetEditorStyle.cpp b/Engine/Plugins/Experimental/ToolPresets/Source/PresetEditor/Private/PresetEditorStyle.cpp
deleted file mode 100644
diff --git a/Engine/Plugins/Experimental/ToolPresets/Source/PresetEditor/Private/PresetEditorStyle.h b/Engine/Plugins/Experimental/ToolPresets/Source/PresetEditor/Private/PresetEditorStyle.h
deleted file mode 100644
diff --git a/Engine/Plugins/Experimental/ToolPresets/Source/PresetEditor/Private/PresetSettings.cpp b/Engine/Plugins/Experimental/ToolPresets/Source/PresetEditor/Private/PresetSettings.cpp
deleted file mode 100644
diff --git a/Engine/Plugins/Experimental/ToolPresets/Source/PresetEditor/Private/Widgets/SPresetManager.cpp b/Engine/Plugins/Experimental/ToolPresets/Source/PresetEditor/Private/Widgets/SPresetManager.cpp
deleted file mode 100644
diff --git a/Engine/Plugins/Experimental/ToolPresets/Source/PresetEditor/Private/Widgets/SPresetManager.h b/Engine/Plugins/Experimental/ToolPresets/Source/PresetEditor/Private/Widgets/SPresetManager.h
deleted file mode 100644
diff --git a/Engine/Plugins/Experimental/ToolPresets/Source/PresetEditor/Public/IPresetEditorModule.h b/Engine/Plugins/Experimental/ToolPresets/Source/PresetEditor/Public/IPresetEditorModule.h
deleted file mode 100644
diff --git a/Engine/Plugins/Experimental/ToolPresets/Source/PresetEditor/Public/PresetSettings.h b/Engine/Plugins/Experimental/ToolPresets/Source/PresetEditor/Public/PresetSettings.h
deleted file mode 100644
diff --git a/Engine/Plugins/Experimental/ToolPresets/Source/ToolPresetAsset/Private/AssetDefinition_ToolPresetAsset.cpp b/Engine/Plugins/Experimental/ToolPresets/Source/ToolPresetAsset/Private/AssetDefinition_ToolPresetAsset.cpp
new file mode 100644
--- /dev/null
+++ b/Engine/Plugins/Experimental/ToolPresets/Source/ToolPresetAsset/Private/AssetDefinition_ToolPresetAsset.cpp
@@ -0,0 +1,7 @@
+// Copyright Epic Games, Inc. All Rights Reserved.
+
+#include "AssetDefinition_ToolPresetAsset.h"
+
+#define LOCTEXT_NAMESPACE "AssetDefinition_ToolPresetAsset"
+
+#undef LOCTEXT_NAMESPACE
diff --git a/Engine/Plugins/Experimental/ToolPresets/Source/ToolPresetAsset/Private/AssetDefinition_ToolPresetAsset.h b/Engine/Plugins/Experimental/ToolPresets/Source/ToolPresetAsset/Private/AssetDefinition_ToolPresetAsset.h
new file mode 100644
--- /dev/null
+++ b/Engine/Plugins/Experimental/ToolPresets/Source/ToolPresetAsset/Private/AssetDefinition_ToolPresetAsset.h
@@ -0,0 +1,30 @@
+// Copyright Epic Games, Inc. All Rights Reserved.
+
+#pragma once
+
+#include "ToolPresetAsset.h"
+#include "AssetDefinitionDefault.h"
+
+#include "AssetDefinition_ToolPresetAsset.generated.h"
+
+enum class EAssetCommandResult : uint8;
+struct FAssetCategoryPath;
+struct FAssetOpenArgs;
+
+UCLASS()
+class UAssetDefinition_InteractiveToolsPresetCollectionAsset : public UAssetDefinitionDefault
+{
+	GENERATED_BODY()
+
+public:
+	// UAssetDefinition Begin
+	virtual FText GetAssetDisplayName() const override { return NSLOCTEXT("AssetTypeActions", "AssetTypeActions_InteractiveToolsPresetCollectionAsset", "Interactive Tools Preset Collection"); }
+	virtual FLinearColor GetAssetColor() const override { return FLinearColor(FColor(175, 0, 128)); }
+	virtual TSoftClassPtr<UObject> GetAssetClass() const override { return UInteractiveToolsPresetCollectionAsset::StaticClass(); }
+	virtual TConstArrayView<FAssetCategoryPath> GetAssetCategories() const override
+	{
+		return TConstArrayView<FAssetCategoryPath>();
+	}
+	virtual FText GetObjectDisplayNameText(UObject* Object) const override { return FText::FromString(TEXT("UInteractiveToolsPresetCollectionAsset")); }
+	// UAssetDefinition End
+};
diff --git a/Engine/Plugins/Experimental/ToolPresets/Source/ToolPresetAsset/Private/ToolPresetAsset.cpp b/Engine/Plugins/Experimental/ToolPresets/Source/ToolPresetAsset/Private/ToolPresetAsset.cpp
new file mode 100644
--- /dev/null
+++ b/Engine/Plugins/Experimental/ToolPresets/Source/ToolPresetAsset/Private/ToolPresetAsset.cpp
@@ -0,0 +1,75 @@
+// Copyright Epic Games, Inc. All Rights Reserved.
+
+#include "ToolPresetAsset.h"
+#include "JsonObjectConverter.h"
+
+void FInteractiveToolPresetDefinition::SetStoredPropertyData(TArray<UObject*>& Properties)
+{
+	FJsonObjectWrapper JsonWrapper;
+
+	for (UObject* PropertySet : Properties)
+	{
+		TSharedPtr<FJsonObject> PropertyJsonObject = MakeShared<FJsonObject>();
+
+		if (PropertyJsonObject)
+		{
+			for (FProperty* Prop : TFieldRange<FProperty>(PropertySet->GetClass()))
+			{
+#if WITH_EDITOR
+				if (! (Prop->HasMetaData(TEXT("TransientToolProperty")) ||
+					   Prop->GetPropertyFlags() & EPropertyFlags::CPF_SkipSerialization ||
+					   Prop->GetPropertyFlags() & EPropertyFlags::CPF_Transient) )
+#endif
+				{
+					TSharedPtr<FJsonValue> JsonProp = FJsonObjectConverter::UPropertyToJsonValue(Prop, Prop->ContainerPtrToValuePtr<void>(PropertySet));
+					FString FieldName;
+					Prop->GetName(FieldName);
+					PropertyJsonObject->SetField(FieldName, JsonProp);
+				}
+			}			
+		}
+
+		FString ClassName;
+		PropertySet->GetClass()->GetName(ClassName);
+		JsonWrapper.JsonObject->SetObjectField(ClassName, PropertyJsonObject);
+	}
+	
+	JsonWrapper.JsonObjectToString(StoredProperties);
+}
+
+void FInteractiveToolPresetDefinition::LoadStoredPropertyData(TArray<UObject*>& Properties)
+{
+	FJsonObjectWrapper JsonWrapper;
+	JsonWrapper.JsonObjectFromString(StoredProperties);
+	
+	for (UObject* PropertySet : Properties)
+	{
+		FString ClassName;
+		PropertySet->GetClass()->GetName(ClassName);
+		TSharedPtr<FJsonObject> PropertyJsonObject = JsonWrapper.JsonObject->GetObjectField(ClassName);
+
+		if (PropertyJsonObject)
+		{
+			for (FProperty* Prop : TFieldRange<FProperty>(PropertySet->GetClass()))
+			{
+				FString FieldName;
+				Prop->GetName(FieldName);
+				TSharedPtr<FJsonValue>* JsonField = PropertyJsonObject->Values.Find(FieldName);
+				if (JsonField)
+				{
+					if (FJsonObjectConverter::JsonValueToUProperty(*JsonField, Prop, Prop->ContainerPtrToValuePtr<void>(PropertySet)))
+					{
+						FPropertyChangedEvent ChangeEvent(Prop, EPropertyChangeType::ValueSet);
+						PropertySet->PostEditChangeProperty(ChangeEvent);
+					}
+				}
+			}
+		}
+	}
+}
+
+bool FInteractiveToolPresetDefinition::IsValid() const
+{
+	FJsonObjectWrapper JsonWrapper;
+	return JsonWrapper.JsonObjectFromString(StoredProperties);
+}
\ No newline at end of file
diff --git a/Engine/Plugins/Experimental/ToolPresets/Source/ToolPresetAsset/Private/ToolPresetAssetModule.cpp b/Engine/Plugins/Experimental/ToolPresets/Source/ToolPresetAsset/Private/ToolPresetAssetModule.cpp
new file mode 100644
--- /dev/null
+++ b/Engine/Plugins/Experimental/ToolPresets/Source/ToolPresetAsset/Private/ToolPresetAssetModule.cpp
@@ -0,0 +1,66 @@
+// Copyright Epic Games, Inc. All Rights Reserved.
+
+#include "Modules/ModuleInterface.h"
+
+#include "Modules/ModuleManager.h"
+#include "PropertyEditorModule.h"
+#include "ToolPresetAsset.h"
+
+#define LOCTEXT_NAMESPACE "FToolPresetAssetModule"
+
+/**
+ * Implements the PresetAsset module.
+ */
+class FToolPresetAssetModule
+	: public IModuleInterface
+{
+public:
+
+	//~ IModuleInterface interface
+
+	virtual void StartupModule() override;
+	virtual void ShutdownModule() override;
+	virtual void OnPostEngineInit();
+	virtual bool SupportsDynamicReloading() override
+	{
+		return true;
+	}
+
+protected:
+
+	TArray<FName> ClassesToUnregisterOnShutdown;
+};
+
+
+// This function will execute after your module is loaded into memory; the exact timing is specified in the .uplugin file per-module
+void FToolPresetAssetModule::StartupModule()
+{
+	FCoreDelegates::OnPostEngineInit.AddRaw(this, &FToolPresetAssetModule::OnPostEngineInit);
+}
+
+// This function may be called during shutdown to clean up your module.  For modules that support dynamic reloading,
+// we call this function before unloading the module.
+void FToolPresetAssetModule::ShutdownModule()
+{
+	// Unregister customizations
+	FPropertyEditorModule* PropertyEditorModule = FModuleManager::GetModulePtr<FPropertyEditorModule>("PropertyEditor");
+	if (PropertyEditorModule)
+	{
+		for (FName ClassName : ClassesToUnregisterOnShutdown)
+		{
+			PropertyEditorModule->UnregisterCustomClassLayout(ClassName);
+		}
+	}
+}
+
+void FToolPresetAssetModule::OnPostEngineInit()
+{
+	ClassesToUnregisterOnShutdown.Reset();
+
+	FPropertyEditorModule& PropertyModule = FModuleManager::LoadModuleChecked<FPropertyEditorModule>("PropertyEditor");
+}
+
+
+IMPLEMENT_MODULE(FToolPresetAssetModule, PresetAsset);
+
+#undef LOCTEXT_NAMESPACE
diff --git a/Engine/Plugins/Experimental/ToolPresets/Source/ToolPresetAsset/Private/ToolPresetAssetSubsystem.cpp b/Engine/Plugins/Experimental/ToolPresets/Source/ToolPresetAsset/Private/ToolPresetAssetSubsystem.cpp
new file mode 100644
--- /dev/null
+++ b/Engine/Plugins/Experimental/ToolPresets/Source/ToolPresetAsset/Private/ToolPresetAssetSubsystem.cpp
@@ -0,0 +1,57 @@
+// Copyright Epic Games, Inc. All Rights Reserved.
+
+#include "ToolPresetAssetSubsystem.h"
+
+#include "EditorConfigSubsystem.h"
+#include "ToolPresetAsset.h"
+
+#define LOCTEXT_NAMESPACE "ToolPresetAssetSubsystem"
+
+void UToolPresetAssetSubsystem::Initialize(FSubsystemCollectionBase& Collection)
+{
+	Collection.InitializeDependency(UEditorConfigSubsystem::StaticClass());
+
+	InitializeDefaultCollection();
+}
+
+void UToolPresetAssetSubsystem::Deinitialize()
+{
+	if (DefaultCollection)
+	{
+		DefaultCollection->SaveEditorConfig();
+		DefaultCollection = nullptr;
+	}	
+}
+
+UInteractiveToolsPresetCollectionAsset* UToolPresetAssetSubsystem::GetDefaultCollection()
+{
+	return DefaultCollection;
+}
+
+bool UToolPresetAssetSubsystem::SaveDefaultCollection()
+{
+	if (DefaultCollection)
+	{
+		return DefaultCollection->SaveEditorConfig();
+	}
+	return false;
+}
+
+
+void UToolPresetAssetSubsystem::InitializeDefaultCollection()
+{
+	/*
+	* We're storing the default collection as a JSON file instead of an asset
+	* on disk for a few reasons. First it avoids issues around automatically
+	* creating assets, both from avoiding build system issues and from a more
+	* philosophical point about requiring user involvement. Second, it helps
+	* compartmentalizing the "default" collection as more of Editor preferences,
+	* rather than a specific collection which has purpose and can be shared around. 
+	*/
+
+	DefaultCollection = NewObject<UInteractiveToolsPresetCollectionAsset>();
+	DefaultCollection->CollectionLabel = LOCTEXT("DefaultCollectionLabel", "Personal Tool Presets (Default)");
+	DefaultCollection->LoadEditorConfig();
+}
+
+#undef LOCTEXT_NAMESPACE
diff --git a/Engine/Plugins/Experimental/ToolPresets/Source/ToolPresetAsset/Public/ToolPresetAsset.h b/Engine/Plugins/Experimental/ToolPresets/Source/ToolPresetAsset/Public/ToolPresetAsset.h
new file mode 100644
--- /dev/null
+++ b/Engine/Plugins/Experimental/ToolPresets/Source/ToolPresetAsset/Public/ToolPresetAsset.h
@@ -0,0 +1,126 @@
+// Copyright Epic Games, Inc. All Rights Reserved.
+
+#pragma once
+
+//#include "Internationalization/Text.h"
+#include "UObject/Object.h"
+#include "UObject/ObjectMacros.h"
+#include "Factories/Factory.h"
+#include "AssetToolsModule.h"
+#include "EditorConfigBase.h"
+
+#include "ToolPresetAsset.generated.h"
+
+#define LOCTEXT_NAMESPACE "ToolPresetAsset"
+
+USTRUCT(BlueprintType)
+struct TOOLPRESETASSET_API FInteractiveToolPresetDefinition
+{
+	GENERATED_BODY()
+public:
+	UPROPERTY(VisibleAnywhere, Category = "ToolPresetAsset", meta = (EditorConfig))
+	FString StoredProperties;
+
+	UPROPERTY(EditAnywhere, Category = "ToolPresetAsset", meta = (EditorConfig))
+	FString Label;
+
+	UPROPERTY(EditAnywhere, Category = "ToolPresetAsset", meta = (EditorConfig))
+	FString Tooltip;
+
+	bool IsValid() const;
+	void SetStoredPropertyData(TArray<UObject*>& Properties);
+	void LoadStoredPropertyData(TArray<UObject*>& Properties);
+};
+
+USTRUCT(BlueprintType)
+struct TOOLPRESETASSET_API FInteractiveToolPresetStore
+{
+	GENERATED_BODY()
+public:
+	UPROPERTY(EditAnywhere, Category = "ToolPresetAsset", meta = (EditorConfig))
+	TArray<FInteractiveToolPresetDefinition> NamedPresets;
+
+	UPROPERTY(VisibleAnywhere, Category = "ToolPresetAsset", meta = (EditorConfig))
+	FText ToolLabel;
+
+	UPROPERTY(VisibleAnywhere, Category = "ToolPresetAsset", meta = (EditorConfig))
+	FSlateBrush ToolIcon;
+};
+
+/**
+ * Implements an asset that can be used to store tool settings as a named preset
+ */
+UCLASS(BlueprintType, hidecategories=(Object), EditorConfig = "UInteractiveToolsPresetCollectionAsset_DefaultCollection")
+class TOOLPRESETASSET_API UInteractiveToolsPresetCollectionAsset 
+	: public UEditorConfigBase
+{
+	GENERATED_BODY()
+
+public:
+
+	// TODO: Currently there are no helper methods within this class, simply providing
+	// raw access to the underlying arrays and maps. This is intentional.
+	// Until the design of the preset concept is more firmly decided, it seems like a waste to
+	// implement a bunch of methods that we don't know if we actually want/need at the end.
+	// Once we are satisifed with the data structure, planned accessors and mutators will include
+	// support for adding, removing, renaming, saving and retrieving presets.
+
+	UPROPERTY(VisibleAnywhere, Category = "ToolPresetAsset", meta = (EditorConfig) )
+	TMap<FString, FInteractiveToolPresetStore > PerToolPresets;
+
+	UPROPERTY(EditAnywhere, Category = "ToolPresetAsset", meta = (EditorConfig))
+	FText CollectionLabel;
+};
+
+
+UCLASS(hidecategories = Object)
+class TOOLPRESETASSET_API UInteractiveToolsPresetCollectionAssetFactory : public UFactory
+{
+	GENERATED_BODY()
+
+public:
+
+	UInteractiveToolsPresetCollectionAssetFactory()
+	{
+		SupportedClass = UInteractiveToolsPresetCollectionAsset::StaticClass();
+
+		bCreateNew = true;
+		bEditAfterNew = true;
+	}
+
+	FText GetDisplayName() const override
+	{
+		return LOCTEXT("DisplayName", "Tool Preset");
+	}
+
+	FText GetToolTip() const override
+	{
+		return LOCTEXT("Tooltip", "Tool Presets capture the state of tool settings for later reloading.");
+	}
+
+	UObject* FactoryCreateNew(UClass* Class, UObject* InParent, FName Name, EObjectFlags Flags, UObject* Context, FFeedbackContext* Warn) override
+	{
+		UInteractiveToolsPresetCollectionAsset* Preset = nullptr;
+		if (ensure(SupportedClass == Class))
+		{
+			ensure(0 != (RF_Public & Flags));
+			Preset = NewObject<UInteractiveToolsPresetCollectionAsset>(InParent, Class, Name, Flags);
+		}
+		return Preset;
+	}
+
+	uint32 GetMenuCategories() const override
+	{
+		IAssetTools& AssetTools = FModuleManager::LoadModuleChecked<FAssetToolsModule>("AssetTools").Get();
+		return AssetTools.RegisterAdvancedAssetCategory("ToolPresets", LOCTEXT("AssetCategoryName", "Tool Presets"));
+	}
+
+	FString GetDefaultNewAssetName() const override
+	{
+		return TEXT("Tool Preset");
+	}
+
+	bool ShouldShowInNewMenu() const override { return true; }
+};
+
+#undef LOCTEXT_NAMESPACE
\ No newline at end of file
diff --git a/Engine/Plugins/Experimental/ToolPresets/Source/ToolPresetAsset/Public/ToolPresetAssetSubsystem.h b/Engine/Plugins/Experimental/ToolPresets/Source/ToolPresetAsset/Public/ToolPresetAssetSubsystem.h
new file mode 100644
--- /dev/null
+++ b/Engine/Plugins/Experimental/ToolPresets/Source/ToolPresetAsset/Public/ToolPresetAssetSubsystem.h
@@ -0,0 +1,33 @@
+// Copyright Epic Games, Inc. All Rights Reserved.
+
+#pragma once
+
+#include "CoreMinimal.h"
+#include "EditorSubsystem.h"
+#include "ToolPresetAssetSubsystem.generated.h"
+
+class UInteractiveToolsPresetCollectionAsset;
+
+/**
+ * Using an editor subsystem allows us to make sure that we have a default preset asset whenever the editor exists
+ *  (and to avoid accidentally trying to make one when it doesn't, such as when running cooking scripts).
+ */
+UCLASS()
+class TOOLPRESETASSET_API UToolPresetAssetSubsystem : public UEditorSubsystem
+{
+	GENERATED_BODY()
+
+public:
+	virtual void Initialize(FSubsystemCollectionBase& Collection) override;
+	virtual void Deinitialize() override;
+
+	UInteractiveToolsPresetCollectionAsset* GetDefaultCollection();
+	bool SaveDefaultCollection();
+
+
+protected:
+	void InitializeDefaultCollection();
+
+	UPROPERTY()
+	TObjectPtr<UInteractiveToolsPresetCollectionAsset> DefaultCollection;
+};
\ No newline at end of file
diff --git a/Engine/Plugins/Experimental/ToolPresets/Source/ToolPresetAsset/ToolPresetAsset.Build.cs b/Engine/Plugins/Experimental/ToolPresets/Source/ToolPresetAsset/ToolPresetAsset.Build.cs
new file mode 100644
--- /dev/null
+++ b/Engine/Plugins/Experimental/ToolPresets/Source/ToolPresetAsset/ToolPresetAsset.Build.cs
@@ -0,0 +1,32 @@
+// Copyright Epic Games, Inc. All Rights Reserved.
+
+using UnrealBuildTool;
+
+public class ToolPresetAsset: ModuleRules
+{
+	public ToolPresetAsset(ReadOnlyTargetRules Target) : base(Target)
+	{
+		PCHUsage = ModuleRules.PCHUsageMode.UseExplicitOrSharedPCHs;
+
+		PublicDependencyModuleNames.AddRange(
+			new string[]
+			{
+				"Core",				
+			});
+
+		PrivateDependencyModuleNames.AddRange(
+			new string[]
+			{
+				"AssetDefinition",
+				"CoreUObject",
+				"EditorConfig",
+				"EditorSubsystem",
+				"Engine",
+				"Json",
+				"JsonUtilities",
+				"Slate",
+				"SlateCore",
+				"UnrealEd",
+			});
+	}
+}
diff --git a/Engine/Plugins/Experimental/ToolPresets/Source/ToolPresetEditor/Private/ToolPresetEditorModule.cpp b/Engine/Plugins/Experimental/ToolPresets/Source/ToolPresetEditor/Private/ToolPresetEditorModule.cpp
new file mode 100644
--- /dev/null
+++ b/Engine/Plugins/Experimental/ToolPresets/Source/ToolPresetEditor/Private/ToolPresetEditorModule.cpp
@@ -0,0 +1,75 @@
+// Copyright Epic Games, Inc. All Rights Reserved.
+
+#include "IToolPresetEditorModule.h"
+
+#include "Delegates/Delegate.h"
+#include "Framework/Docking/TabManager.h"
+#include "Internationalization/Internationalization.h"
+#include "Modules/ModuleManager.h"
+#include "Styling/AppStyle.h"
+#include "Templates/SharedPointer.h"
+#include "Textures/SlateIcon.h"
+#include "ToolPresetEditorStyle.h"
+#include "UObject/NameTypes.h"
+#include "UObject/UObjectGlobals.h"
+#include "Widgets/DeclarativeSyntaxSupport.h"
+#include "Widgets/Docking/SDockTab.h"
+#include "Widgets/SToolPresetManager.h"
+
+#define LOCTEXT_NAMESPACE "ToolPresetEditorModule"
+
+const FName PresetEditorTabName("ToolPreset");
+
+class FToolPresetEditorModule : public IToolPresetEditorModule
+{
+public:
+	virtual void StartupModule() override
+	{
+		FGlobalTabmanager::Get()->RegisterNomadTabSpawner(PresetEditorTabName, FOnSpawnTab::CreateRaw(this, &FToolPresetEditorModule::HandleSpawnPresetEditorTab))
+			.SetDisplayName(NSLOCTEXT("FToolPresetModule", "ToolPresetTabTitle", "Tool Preset Manager"))
+			.SetTooltipText(NSLOCTEXT("FToolPresetModule", "ToolPresetTooltipText", "Open the Tool Preset Manager tab."))
+			.SetIcon(FSlateIcon(FAppStyle::GetAppStyleSetName(), "ToolPreset.TabIcon"))
+			.SetAutoGenerateMenuEntry(false);
+
+		FCoreDelegates::OnPostEngineInit.AddRaw(this, &FToolPresetEditorModule::OnPostEngineInit);
+	}
+	
+	virtual void ShutdownModule() override
+	{
+		FGlobalTabmanager::Get()->UnregisterNomadTabSpawner(PresetEditorTabName);
+
+		FToolPresetEditorStyle::Shutdown();
+	}
+
+	void OnPostEngineInit()
+	{
+		// Register slate style overrides
+		FToolPresetEditorStyle::Initialize();
+	}
+
+	virtual bool SupportsDynamicReloading() override
+	{
+		return true;
+	}
+
+	virtual void ExecuteOpenPresetEditor() override
+	{
+		FGlobalTabmanager::Get()->TryInvokeTab(PresetEditorTabName);
+	}	
+
+private:	
+	/** Handles creating the project settings tab. */
+	TSharedRef<SDockTab> HandleSpawnPresetEditorTab(const FSpawnTabArgs& SpawnTabArgs)
+	{
+		const TSharedRef<SDockTab> DockTab = SNew(SDockTab)
+			.TabRole(ETabRole::NomadTab)
+		[
+			SNew(SToolPresetManager)
+		];
+		return DockTab;
+	}
+};
+
+#undef LOCTEXT_NAMESPACE
+
+IMPLEMENT_MODULE(FToolPresetEditorModule, ToolPresetEditor);
diff --git a/Engine/Plugins/Experimental/ToolPresets/Source/ToolPresetEditor/Private/ToolPresetEditorStyle.cpp b/Engine/Plugins/Experimental/ToolPresets/Source/ToolPresetEditor/Private/ToolPresetEditorStyle.cpp
new file mode 100644
--- /dev/null
+++ b/Engine/Plugins/Experimental/ToolPresets/Source/ToolPresetEditor/Private/ToolPresetEditorStyle.cpp
@@ -0,0 +1,90 @@
+// Copyright Epic Games, Inc. All Rights Reserved.
+
+#include "ToolPresetEditorStyle.h"
+
+#include "Brushes/SlateImageBrush.h"
+#include "Interfaces/IPluginManager.h"
+#include "Misc/Paths.h"
+#include "SlateOptMacros.h"
+#include "Styling/CoreStyle.h"
+#include "Styling/SlateStyle.h"
+#include "Styling/SlateStyleRegistry.h"
+
+#define IMAGE_PLUGIN_BRUSH( RelativePath, ... ) FSlateImageBrush( FToolPresetEditorStyle::InContent( RelativePath, ".png" ), __VA_ARGS__ )
+
+// This is to fix the issue that SlateStyleMacros like IMAGE_BRUSH look for RootToContentDir but StyleSet->RootToContentDir is how this style is set up
+#define RootToContentDir StyleSet->RootToContentDir
+
+FString FToolPresetEditorStyle::InContent(const FString& RelativePath, const ANSICHAR* Extension)
+{
+	static FString ContentDir = IPluginManager::Get().FindPlugin(TEXT("ToolPresets"))->GetContentDir();
+	return (ContentDir / RelativePath) + Extension;
+}
+
+TSharedPtr< FSlateStyleSet > FToolPresetEditorStyle::StyleSet = nullptr;
+TSharedPtr< class ISlateStyle > FToolPresetEditorStyle::Get() { return StyleSet; }
+
+FName FToolPresetEditorStyle::GetStyleSetName()
+{
+	static FName PresetEditorStyleName(TEXT("ToolPresetEditorStyle"));
+	return PresetEditorStyleName;
+}
+
+const FSlateBrush* FToolPresetEditorStyle::GetBrush(FName PropertyName, const ANSICHAR* Specifier)
+{
+	return Get()->GetBrush(PropertyName, Specifier);
+}
+
+
+BEGIN_SLATE_FUNCTION_BUILD_OPTIMIZATION
+
+void FToolPresetEditorStyle::Initialize()
+{
+	// Const icon sizes
+	const FVector2D Icon8x8(8.0f, 8.0f);
+	const FVector2D Icon16x16(16.0f, 16.0f);
+	const FVector2D Icon20x20(20.0f, 20.0f);
+	const FVector2D Icon28x28(28.0f, 28.0f);
+	const FVector2D Icon40x40(40.0f, 40.0f);
+	const FVector2D Icon120(120.0f, 120.0f);
+
+	// Only register once
+	if (StyleSet.IsValid())
+	{
+		return;
+	}
+
+	StyleSet = MakeShareable(new FSlateStyleSet(GetStyleSetName()));
+
+	// If we get asked for something that we don't set, we should default to editor style
+	StyleSet->SetParentStyleName("EditorStyle");
+
+	StyleSet->SetContentRoot(FPaths::EnginePluginsDir() / TEXT("Experimental/ToolPresets/Content"));
+	StyleSet->SetCoreContentRoot(FPaths::EngineContentDir() / TEXT("Slate"));
+
+	// Icons
+	{
+		StyleSet->Set("ManagerIcons.Tools", new IMAGE_PLUGIN_BRUSH("Icons/Primitive_40x", Icon40x40));
+
+
+    }
+
+	FSlateStyleRegistry::RegisterSlateStyle(*StyleSet.Get());
+};
+
+END_SLATE_FUNCTION_BUILD_OPTIMIZATION
+
+#undef IMAGE_PLUGIN_BRUSH
+#undef IMAGE_BRUSH
+#undef BOX_BRUSH
+#undef DEFAULT_FONT
+
+void FToolPresetEditorStyle::Shutdown()
+{
+	if (StyleSet.IsValid())
+	{
+		FSlateStyleRegistry::UnRegisterSlateStyle(*StyleSet.Get());
+		ensure(StyleSet.IsUnique());
+		StyleSet.Reset();
+	}
+}
diff --git a/Engine/Plugins/Experimental/ToolPresets/Source/ToolPresetEditor/Private/ToolPresetEditorStyle.h b/Engine/Plugins/Experimental/ToolPresets/Source/ToolPresetEditor/Private/ToolPresetEditorStyle.h
new file mode 100644
--- /dev/null
+++ b/Engine/Plugins/Experimental/ToolPresets/Source/ToolPresetEditor/Private/ToolPresetEditorStyle.h
@@ -0,0 +1,28 @@
+// Copyright Epic Games, Inc. All Rights Reserved.
+
+#pragma once
+
+#include "Templates/SharedPointer.h"
+
+struct FSlateBrush;
+
+class FToolPresetEditorStyle
+{
+public:
+	static void Initialize();
+
+	static void Shutdown();
+
+	static TSharedPtr< class ISlateStyle > Get();
+
+	static FName GetStyleSetName();
+
+	// use to access icons defined by the style set by name, eg GetBrush("BrushFalloffIcons.Smooth")
+	static const FSlateBrush* GetBrush(FName PropertyName, const ANSICHAR* Specifier = nullptr);
+
+private:
+	static FString InContent(const FString& RelativePath, const ANSICHAR* Extension);
+
+private:
+	static TSharedPtr< class FSlateStyleSet > StyleSet;
+};
diff --git a/Engine/Plugins/Experimental/ToolPresets/Source/ToolPresetEditor/Private/ToolPresetSettings.cpp b/Engine/Plugins/Experimental/ToolPresets/Source/ToolPresetEditor/Private/ToolPresetSettings.cpp
new file mode 100644
--- /dev/null
+++ b/Engine/Plugins/Experimental/ToolPresets/Source/ToolPresetEditor/Private/ToolPresetSettings.cpp
@@ -0,0 +1,21 @@
+// Copyright Epic Games, Inc. All Rights Reserved.
+
+#include "ToolPresetSettings.h"
+
+#include "UObject/ObjectPtr.h"
+
+TObjectPtr<UToolPresetUserSettings> UToolPresetUserSettings::Instance = nullptr;
+
+void UToolPresetUserSettings::Initialize()
+{
+	if (Instance == nullptr)
+	{
+		Instance = NewObject<UToolPresetUserSettings>();
+		Instance->AddToRoot();
+	}
+}
+
+UToolPresetUserSettings* UToolPresetUserSettings::Get()
+{
+	return Instance;
+}
diff --git a/Engine/Plugins/Experimental/ToolPresets/Source/ToolPresetEditor/Private/Widgets/SToolPresetManager.cpp b/Engine/Plugins/Experimental/ToolPresets/Source/ToolPresetEditor/Private/Widgets/SToolPresetManager.cpp
new file mode 100644
--- /dev/null
+++ b/Engine/Plugins/Experimental/ToolPresets/Source/ToolPresetEditor/Private/Widgets/SToolPresetManager.cpp
@@ -0,0 +1,1584 @@
+// Copyright Epic Games, Inc. All Rights Reserved.
+
+#include "Widgets/SToolPresetManager.h"
+
+#include "AssetRegistry/AssetRegistryModule.h"
+#include "ContentBrowserModule.h"
+#include "Editor.h"
+#include "Framework/Commands/GenericCommands.h"
+#include "Framework/MultiBox/MultiBoxBuilder.h"
+#include "IAssetTools.h"
+#include "ISettingsModule.h"
+#include "Modules/ModuleManager.h"
+#include "ObjectTools.h"
+#include "SlateOptMacros.h"
+#include "SNegativeActionButton.h"
+#include "SPositiveActionButton.h"
+#include "SSimpleButton.h"
+#include "Styling/AppStyle.h"
+#include "Templates/SharedPointer.h"
+#include "ToolPresetAsset.h"
+#include "ToolPresetAssetSubsystem.h"
+#include "ToolPresetEditorStyle.h"
+#include "ToolPresetSettings.h"
+#include "UObject/SavePackage.h"
+#include "Widgets/Images/SImage.h"
+#include "Widgets/Input/SButton.h"
+#include "Widgets/Input/SCheckBox.h"
+#include "Widgets/Input/SComboButton.h"
+#include "Widgets/Input/SEditableTextBox.h"
+#include "Widgets/Input/SMultiLineEditableTextBox.h"
+#include "Widgets/Layout/SExpandableArea.h"
+#include "Widgets/Layout/SSpacer.h"
+#include "Widgets/SWidget.h"
+#include "Widgets/Text/STextBlock.h"
+#include "Widgets/Views/SListView.h"
+#include "Widgets/Views/STableRow.h"
+#include "Widgets/Views/STreeView.h"
+
+#define LOCTEXT_NAMESPACE "SToolPresetManager"
+
+DECLARE_DELEGATE_TwoParams(FOnCollectionEnabledCheckboxChanged, TSharedPtr<SToolPresetManager::FToolPresetViewEntry>, ECheckBoxState)
+DECLARE_DELEGATE_TwoParams(FOnPresetLabelChanged, TSharedPtr<SToolPresetManager::FToolPresetViewEntry>, FText)
+DECLARE_DELEGATE_TwoParams(FOnPresetTooltipChanged, TSharedPtr<SToolPresetManager::FToolPresetViewEntry>, FText)
+DECLARE_DELEGATE_OneParam(FOnPresetDeleted, TSharedPtr<SToolPresetManager::FToolPresetViewEntry>)
+DECLARE_DELEGATE_TwoParams(FOnCollectionRenameStarted, TSharedPtr<SToolPresetManager::FToolPresetViewEntry>, TSharedPtr<SEditableTextBox> RenameWidget)
+DECLARE_DELEGATE_TwoParams(FOnCollectionRenameEnded, TSharedPtr<SToolPresetManager::FToolPresetViewEntry>, const FText& NewText)
+
+namespace UE::ToolPresetEditor::Private
+{
+	template <typename AssetType>
+	void GetObjectsOfClass(TArray<FSoftObjectPath>& OutArray)
+	{
+		FAssetRegistryModule& AssetRegistryModule = FModuleManager::LoadModuleChecked<FAssetRegistryModule>("AssetRegistry");
+		TArray<FAssetData> AssetData;
+
+		FARFilter Filter;
+		Filter.ClassPaths.Add(AssetType::StaticClass()->GetClassPathName());
+		Filter.PackagePaths.Add(FName("/ToolPresets"));
+		Filter.bRecursiveClasses = false;	
+		Filter.bRecursivePaths = true;
+		Filter.bIncludeOnlyOnDiskAssets = false;
+		
+		AssetRegistryModule.Get().GetAssets(Filter, AssetData);
+
+		for (int i = 0; i < AssetData.Num(); i++)
+		{
+			AssetType* Object = Cast<AssetType>(AssetData[i].GetAsset());		
+			if (Object)
+			{				
+				OutArray.Add(Object->GetPathName());
+			}
+		}
+	}
+
+	template<typename ItemType> class SCollectionTableRow;
+
+	template<typename ItemType>
+	class SCollectionTableRow : public STableRow< ItemType>
+	{
+		typedef SCollectionTableRow< ItemType > FSuperRowType;
+		typedef typename STableRow<ItemType>::FArguments FTableRowArgs;
+		typedef SToolPresetManager::FToolPresetViewEntry::EEntryType EEntryType;
+
+	public:
+
+		SLATE_BEGIN_ARGS(SCollectionTableRow) { }
+		    SLATE_ARGUMENT(TSharedPtr<SToolPresetManager::FToolPresetViewEntry>, ViewEntry)
+			SLATE_EVENT(FOnCollectionEnabledCheckboxChanged, OnCollectionEnabledCheckboxChanged)
+			SLATE_EVENT(FOnCollectionRenameStarted, OnCollectionRenameStarted)
+			SLATE_EVENT(FOnCollectionRenameEnded, OnCollectionRenameEnded)
+		SLATE_END_ARGS()
+
+		void Construct(const FArguments& InArgs, const TSharedRef<STableViewBase>& InOwnerTableView)
+		{
+			FTableRowArgs Args = FTableRowArgs()
+				.Style(&FAppStyle::Get().GetWidgetStyle<FTableRowStyle>("TableView.Row"))
+				.ExpanderStyleSet(& FCoreStyle::Get());
+
+			ViewEntry = InArgs._ViewEntry;
+			OnCollectionEnabledCheckboxChanged = InArgs._OnCollectionEnabledCheckboxChanged;
+			OnCollectionRenameStarted = InArgs._OnCollectionRenameStarted;
+			OnCollectionRenameEnded = InArgs._OnCollectionRenameEnded;
+
+			STableRow<ItemType>::Construct(Args, InOwnerTableView);
+		}
+
+		BEGIN_SLATE_FUNCTION_BUILD_OPTIMIZATION
+		virtual void ConstructChildren(ETableViewMode::Type InOwnerTableMode, const TAttribute<FMargin>& InPadding, const TSharedRef<SWidget>& InContent) override
+		{
+			STableRow<ItemType>::Content = InContent;
+
+			TSharedPtr<class ITableRow> ThisTableRow = this->SharedThis(this);
+
+			if(InOwnerTableMode == ETableViewMode::Tree)
+			{
+
+				// Rows in a TreeView need an expander button and some indentation
+				this->ChildSlot
+				[
+					SNew(SHorizontalBox)
+
+					+ SHorizontalBox::Slot()
+					.AutoWidth()
+					.HAlign(HAlign_Right)
+					.VAlign(VAlign_Fill)
+					[
+						SAssignNew(EnabledWidget, SCheckBox)
+						.Visibility_Lambda([this]()
+						{
+							return ViewEntry->EntryType == EEntryType::Collection ? EVisibility::Visible : EVisibility::Collapsed;
+						})
+						.IsChecked_Lambda([this]()
+						{
+							return ViewEntry->bEnabled ? ECheckBoxState::Checked : ECheckBoxState::Unchecked;
+						})
+						.OnCheckStateChanged_Lambda([this](ECheckBoxState State)
+						{
+							if (OnCollectionEnabledCheckboxChanged.IsBound())
+							{
+								OnCollectionEnabledCheckboxChanged.Execute(ViewEntry, State);
+							}
+						})
+				
+					]
+
+					+ SHorizontalBox::Slot()
+					.AutoWidth()
+					.HAlign(HAlign_Right)
+					.VAlign(VAlign_Fill)
+					[
+						SAssignNew(STableRow< ItemType>::ExpanderArrowWidget, SExpanderArrow, ThisTableRow)
+						.StyleSet(STableRow< ItemType>::ExpanderStyleSet)
+						.ShouldDrawWires(false)
+					]
+
+					+ SHorizontalBox::Slot()
+						.AutoWidth()						
+						.Padding(5.0f)
+						[
+							SNew(SImage)
+							.Visibility_Lambda([this]()
+							{
+								return ViewEntry->EntryType == EEntryType::Tool ? EVisibility::Visible : EVisibility::Collapsed;
+							})
+							.Image(&ViewEntry->EntryIcon)
+						]
+
+					+ SHorizontalBox::Slot()
+						.FillWidth(1)
+						.Padding(5.0f)
+						[
+							SNew(SHorizontalBox)
+
+							+ SHorizontalBox::Slot()
+							[
+								SNew(STextBlock)
+								.Text(ViewEntry->EntryLabel)
+								.Visibility_Lambda([this]()
+								{
+									return ViewEntry->bIsRenaming ? EVisibility::Collapsed : EVisibility::Visible;
+								})							
+								.Font_Lambda([this]()
+								{
+									return FCoreStyle::Get().GetWidgetStyle<FTextBlockStyle>("NormalText").Font;							
+								})
+							]
+
+							+ SHorizontalBox::Slot()
+							[
+								SAssignNew(CollectionRenameBox, SEditableTextBox)
+								.Text(ViewEntry->EntryLabel)
+								
+								.Visibility_Lambda([this]()
+								{
+									if (!bIsEntryBeingRenamed && ViewEntry->bIsRenaming)
+									{
+										OnCollectionRenameStarted.ExecuteIfBound(ViewEntry, CollectionRenameBox);
+									}
+									bIsEntryBeingRenamed = ViewEntry->bIsRenaming;
+									return ViewEntry->bIsRenaming ? EVisibility::Visible : EVisibility::Collapsed;
+								})
+								.OnTextCommitted_Lambda([this](const FText & NewText, ETextCommit::Type CommitStatus)
+								{
+									OnCollectionRenameEnded.ExecuteIfBound(ViewEntry, NewText);
+								})
+								.Font_Lambda([this]()
+								{
+									return FCoreStyle::Get().GetWidgetStyle<FTextBlockStyle>("NormalText").Font;							
+								})
+							]
+						]
+
+
+					+ SHorizontalBox::Slot()
+					.AutoWidth()
+					[
+						SNew(SSpacer)
+					
+					]
+
+					+ SHorizontalBox::Slot()
+						.AutoWidth()
+						.HAlign(EHorizontalAlignment::HAlign_Right)
+						.Padding(5.0f)
+						[
+							SNew(STextBlock)
+							.Text(TAttribute<FText>::CreateLambda([this]()
+							{
+								return FText::AsNumber(ViewEntry->Count);
+							}))
+						]
+				];
+			}
+		}
+		END_SLATE_FUNCTION_BUILD_OPTIMIZATION
+
+	private:
+		TSharedPtr<SToolPresetManager::FToolPresetViewEntry> ViewEntry;
+		TSharedPtr<SCheckBox> EnabledWidget;
+		TSharedPtr<SEditableTextBox> CollectionRenameBox;
+
+		bool bIsEntryBeingRenamed = false;
+
+		FOnCollectionEnabledCheckboxChanged OnCollectionEnabledCheckboxChanged;
+		FOnCollectionRenameStarted OnCollectionRenameStarted;
+		FOnCollectionRenameEnded OnCollectionRenameEnded;
+	};
+
+	template<typename ItemType>
+	class SToolPresetTableRow : public SMultiColumnTableRow< ItemType>
+	{
+		typedef SToolPresetManager::FToolPresetViewEntry::EEntryType EEntryType;
+		using typename SMultiColumnTableRow< ItemType>::FSuperRowType;
+
+	public:
+
+		SLATE_BEGIN_ARGS(SToolPresetTableRow) { }
+		SLATE_ARGUMENT(TSharedPtr<SToolPresetManager::FToolPresetViewEntry>, ViewEntry)
+		SLATE_EVENT(FOnPresetDeleted, OnPresetDeleted)
+		SLATE_END_ARGS()
+
+		void Construct(const FArguments& InArgs, const TSharedRef<STableViewBase>& InOwnerTableView)
+		{
+			typename FSuperRowType::FArguments Args = typename FSuperRowType::FArguments()
+				.ExpanderStyleSet(&FCoreStyle::Get());
+
+			OnPresetDeleted = InArgs._OnPresetDeleted;
+			ViewEntry = InArgs._ViewEntry;
+
+			SMultiColumnTableRow<ItemType>::Construct(Args, InOwnerTableView);
+		}
+
+		BEGIN_SLATE_FUNCTION_BUILD_OPTIMIZATION
+		virtual void ConstructChildren(ETableViewMode::Type InOwnerTableMode, const TAttribute<FMargin>& InPadding, const TSharedRef<SWidget>& InContent) override
+		{
+			STableRow<ItemType>::Content = InContent;
+
+			TSharedPtr<class ITableRow> ThisTableRow = this->SharedThis(this);
+
+			if (ViewEntry->EntryType == EEntryType::Tool)
+			{
+				this->ChildSlot				
+				[
+				SNew(SHorizontalBox)
+
+				+ SHorizontalBox::Slot()
+					.AutoWidth()
+					.Padding(5.0f)
+					[
+						SNew(SImage)
+						.Visibility_Lambda([this]()
+						{
+							return ViewEntry->EntryType == EEntryType::Tool ? EVisibility::Visible : EVisibility::Collapsed;
+						})
+						.Image(&ViewEntry->EntryIcon)
+					]
+
+				+ SHorizontalBox::Slot()
+					.FillWidth(1)
+					.Padding(5.0f)
+					[
+						SNew(STextBlock)
+						.Text(ViewEntry->EntryLabel)
+					]
+
+				+ SHorizontalBox::Slot()
+					.AutoWidth()
+					[
+						SNew(SSpacer)
+
+					]
+
+				+ SHorizontalBox::Slot()
+					.AutoWidth()
+					.HAlign(EHorizontalAlignment::HAlign_Right)
+					.Padding(5.0f)
+					[
+						SNew(STextBlock)
+						.Text(TAttribute<FText>::CreateLambda([this]()
+						{
+							return FText::AsNumber(ViewEntry->Count);
+						}))
+					]
+				];
+			}
+			else
+			{
+				ensure(ViewEntry->EntryType == EEntryType::Preset);
+
+				this->ChildSlot
+				.Padding(InPadding)
+				[
+					InContent
+				];
+			}
+			
+		}
+
+		virtual TSharedRef<SWidget> GenerateWidgetForColumn(const FName& InColumnName)
+		{
+			if (InColumnName == "Label")
+			{
+				return SNew(SHorizontalBox)
+				+ SHorizontalBox::Slot()
+					.VAlign(VAlign_Center)
+					.Padding(FMargin(4.0f, 0.0f))
+					[
+						SNew(STextBlock)
+						.Text(TAttribute<FText>::CreateLambda([this]()
+						{
+							return FText::FromString(ViewEntry->PresetLabel);
+						}))
+					];
+			}
+			else if (InColumnName == "Tooltip")
+			{
+				return SNew(SHorizontalBox)
+				+ SHorizontalBox::Slot()
+					.VAlign(VAlign_Center)
+					.Padding(FMargin(4.0f, 0.0f))
+					[
+						SNew(STextBlock)
+						.Text(TAttribute<FText>::CreateLambda([this]()
+						{
+							return FText::FromString(ViewEntry->PresetTooltip);
+						}))
+					];
+			}
+			else if (InColumnName == "Tool")
+			{
+				return SNew(SBox)
+					.VAlign(VAlign_Center)
+					.Padding(FMargin(5.0f, 5.0f))
+					[
+						SNew(SImage)
+						.Image(&ViewEntry->EntryIcon)
+						.DesiredSizeOverride(FVector2D(16, 16))
+					];
+			}
+			else if (InColumnName == "Delete")
+			{
+				return SNew(SNegativeActionButton)
+					.Icon(FAppStyle::GetBrush("Icons.Delete"))						
+					.OnClicked_Lambda([this]()
+					{
+						OnPresetDeleted.ExecuteIfBound(ViewEntry); return FReply::Handled();
+					})
+					.Visibility_Lambda([this]()
+					{
+						return this->IsHovered() ? EVisibility::Visible : EVisibility::Hidden;
+					});
+			}
+
+			return SNullWidget::NullWidget;
+
+		}
+		END_SLATE_FUNCTION_BUILD_OPTIMIZATION
+
+	private:
+		TSharedPtr<SToolPresetManager::FToolPresetViewEntry> ViewEntry;
+		FOnPresetDeleted                     OnPresetDeleted;
+	};
+
+}
+
+
+
+/* SToolPresetManager interface
+ *****************************************************************************/
+
+BEGIN_SLATE_FUNCTION_BUILD_OPTIMIZATION
+void SToolPresetManager::Construct( const FArguments& InArgs )
+{
+	UToolPresetUserSettings::Initialize();
+	BindCommands();
+
+	UserSettings = UToolPresetUserSettings::Get();
+	if (UserSettings.IsValid())
+	{
+		UserSettings->LoadEditorConfig();
+	}
+
+	ChildSlot
+	[
+		SNew(SVerticalBox)
+
+		+ SVerticalBox::Slot()
+		.FillHeight(1.0)
+		//.AutoHeight()
+		[
+				SAssignNew(Splitter, SSplitter)
+					.Orientation(EOrientation::Orient_Horizontal)
+					
+					+ SSplitter::Slot()
+					.Value(0.4f)
+					.Resizable(false)
+						[
+							SNew(SVerticalBox)		
+
+								+ SVerticalBox::Slot()
+								.AutoHeight()
+								[
+									SNew(SBorder)
+									.BorderImage(&FAppStyle::Get().GetWidgetStyle<FTableColumnHeaderStyle>("TableView.Header.Column").NormalBrush)
+									[
+										SNew(SHorizontalBox)
+									
+										+ SHorizontalBox::Slot()
+										.AutoWidth()
+										.HAlign(HAlign_Right)
+										.VAlign(VAlign_Fill)
+										[
+											SAssignNew(UserCollectionsExpander, SButton)
+											.ButtonStyle(FCoreStyle::Get(), "NoBorder" )
+											.VAlign(VAlign_Center)
+											.HAlign(HAlign_Center)
+											.ClickMethod( EButtonClickMethod::MouseDown )
+											.OnClicked_Lambda([this]()
+											{
+												bAreUserCollectionsExpanded = !bAreUserCollectionsExpanded; return FReply::Handled();
+											})
+											.ContentPadding(0.f)
+											.ForegroundColor( FSlateColor::UseForeground() )
+											.IsFocusable( false )
+											[
+												SNew(SImage)
+												.Image( this, &SToolPresetManager::GetUserCollectionsExpanderImage )
+												.ColorAndOpacity( FSlateColor::UseSubduedForeground() )
+											]
+										]
+									
+
+										+ SHorizontalBox::Slot()
+										.FillWidth(1.0)
+										.Padding(5.0f, 5.0f, 5.0f, 5.0f)
+										.HAlign(HAlign_Left)
+										.VAlign(VAlign_Center)
+										[
+											SNew(STextBlock)										
+											.Text(LOCTEXT("UserPresetLabels", "User Preset Collections"))
+											.Font(FSlateFontInfo(FCoreStyle::GetDefaultFont(), 13, "Bold"))
+										]
+
+										+ SHorizontalBox::Slot()
+										.AutoWidth()
+										.HAlign(HAlign_Right)
+										.VAlign(VAlign_Center)
+										.Padding(16.0f, 8.0f, 8.0f, 8.0f)
+										[
+											SAssignNew(AddUserPresetButton, SPositiveActionButton)
+											.ToolTipText(LOCTEXT("AddUserPresetCollection", "Add User Preset Collection"))
+											.Icon(FAppStyle::GetBrush("Icons.Plus"))
+											.OnClicked_Lambda([this]()
+											{
+												AddNewUserPresetCollection(); return FReply::Handled();
+											})
+										]
+									]
+								]
+
+								+ SVerticalBox::Slot()
+								.AutoHeight()
+								[
+			
+											SAssignNew(EditorPresetCollectionTreeView, STreeView<TSharedPtr<FToolPresetViewEntry> >)
+												.TreeItemsSource(&EditorCollectionsDataList)
+												.SelectionMode(ESelectionMode::Single)
+												.OnGenerateRow(this, &SToolPresetManager::HandleTreeGenerateRow)
+												.OnGetChildren(this, &SToolPresetManager::HandleTreeGetChildren)
+												.OnSelectionChanged(this, &SToolPresetManager::HandleEditorTreeSelectionChanged)
+												.Visibility_Lambda([this]()
+												{
+													return bAreUserCollectionsExpanded ? EVisibility::Visible : EVisibility::Collapsed;
+												})
+												.HeaderRow
+												(
+													SNew(SHeaderRow)
+													.Visibility(EVisibility::Collapsed)
+
+													+ SHeaderRow::Column("Collection")
+													.FixedWidth(150.0f)
+													.HeaderContent()
+													[
+														SNew(STextBlock)
+														.Text(LOCTEXT("ToolPresetManagerCollectionTitleHeader", "Collection"))
+													]
+
+
+												)											
+								]
+
+								+ SVerticalBox::Slot()
+									.FillHeight(1.0f)
+								[
+			
+											SAssignNew(UserPresetCollectionTreeView, STreeView<TSharedPtr<FToolPresetViewEntry> >)
+												.TreeItemsSource(&UserCollectionsDataList)
+												.SelectionMode(ESelectionMode::Single)
+												.OnGenerateRow(this, &SToolPresetManager::HandleTreeGenerateRow)
+												.OnGetChildren(this, &SToolPresetManager::HandleTreeGetChildren)
+												.OnSelectionChanged(this, &SToolPresetManager::HandleUserTreeSelectionChanged)
+												.OnContextMenuOpening(this, &SToolPresetManager::OnGetCollectionContextMenuContent)
+												.Visibility_Lambda([this]()
+												{
+													return bAreUserCollectionsExpanded ? EVisibility::Visible : EVisibility::Collapsed;
+												})
+												.HeaderRow
+												(
+													SNew(SHeaderRow)
+													.Visibility(EVisibility::Collapsed)
+
+													+ SHeaderRow::Column("Collection")
+													.FixedWidth(150.0f)
+													.HeaderContent()
+													[
+														SNew(STextBlock)
+														.Text(LOCTEXT("ToolPresetManagerCollectionTitleHeader", "Collection"))
+													]
+
+
+												)											
+								]
+
+								+ SVerticalBox::Slot()
+									.AutoHeight()									
+								[
+									SNew(SBorder)
+									.BorderImage(&FAppStyle::Get().GetWidgetStyle<FTableColumnHeaderStyle>("TableView.Header.Column").NormalBrush)
+									[
+										SNew(SHorizontalBox)
+										+ SHorizontalBox::Slot()
+										.AutoWidth()
+										.HAlign(HAlign_Right)
+										.VAlign(VAlign_Fill)
+										[
+											SAssignNew(ProjectCollectionsExpander, SButton)
+											.ButtonStyle(FCoreStyle::Get(), "NoBorder" )
+											.VAlign(VAlign_Center)
+											.HAlign(HAlign_Center)
+											.ClickMethod( EButtonClickMethod::MouseDown )
+											.OnClicked_Lambda([this]()
+											{
+												bAreProjectCollectionsExpanded = !bAreProjectCollectionsExpanded; return FReply::Handled();
+											})
+											.ContentPadding(0.f)
+											.ForegroundColor( FSlateColor::UseForeground() )
+											.IsFocusable( false )
+											[
+												SNew(SImage)
+												.Image( this, &SToolPresetManager::GetProjectCollectionsExpanderImage )
+												.ColorAndOpacity( FSlateColor::UseSubduedForeground() )
+											]
+										]
+
+										+ SHorizontalBox::Slot()
+										.FillWidth(1.0)
+										.Padding(5.0f, 5.0f, 5.0f, 5.0f)
+										.HAlign(HAlign_Left)
+										.VAlign(VAlign_Center)
+										[
+											SNew(STextBlock)
+											.Text(LOCTEXT("ProjectPresetLabels", "Project Preset Collections"))
+											.Font(FSlateFontInfo(FCoreStyle::GetDefaultFont(), 13, "Bold"))
+										]
+
+										+ SHorizontalBox::Slot()
+										.AutoWidth()
+										.HAlign(HAlign_Right)
+										.VAlign(VAlign_Center)
+										.Padding(16.0f, 8.0f, 8.0f, 8.0f)
+										[
+											SNew(SSimpleButton)
+											.ToolTipText(LOCTEXT("OpenProjectSettingsPresets", "Open Project Settings for Presets"))
+											.Icon(FAppStyle::GetBrush("Icons.Settings"))
+											.OnClicked_Lambda([this]() 
+											{ 
+												if (ISettingsModule* SettingsModule = FModuleManager::GetModulePtr<ISettingsModule>("Settings"))
+												{
+													SettingsModule->ShowViewer("Project", "Plugins", "Interactive Tool Presets");
+												}
+												return FReply::Handled();
+											})
+										]
+									]
+								]
+
+								+ SVerticalBox::Slot()
+									.FillHeight(1.0f)
+								[
+										
+											SNew(SVerticalBox)
+											.Visibility_Lambda([this]()
+											{
+												return bAreProjectCollectionsExpanded ? EVisibility::Visible : EVisibility::Collapsed;
+											})
+
+											+ SVerticalBox::Slot()
+											.FillHeight(1.0f)
+											[
+
+											SAssignNew(ProjectPresetCollectionTreeView, STreeView<TSharedPtr<FToolPresetViewEntry> >)
+											.Visibility(this, &SToolPresetManager::ProjectPresetCollectionsVisibility)
+												.ItemHeight(32.0f)
+												.TreeItemsSource(&ProjectCollectionsDataList)
+												.SelectionMode(ESelectionMode::Single)
+												.OnGenerateRow(this, &SToolPresetManager::HandleTreeGenerateRow)
+												.OnGetChildren(this, &SToolPresetManager::HandleTreeGetChildren)
+												.OnSelectionChanged(this, &SToolPresetManager::HandleTreeSelectionChanged)
+												.OnContextMenuOpening(this, &SToolPresetManager::OnGetCollectionContextMenuContent)
+												.HeaderRow
+												(
+													SNew(SHeaderRow)
+													.Visibility(EVisibility::Collapsed)
+
+													+ SHeaderRow::Column("Collection")
+													.FixedWidth(150.0f)
+													.HeaderContent()
+													[
+														SNew(STextBlock)
+														.Text(LOCTEXT("ToolPresetManagerCollectionTitleHeader", "Collection"))
+													]
+
+
+												)											
+											]
+
+											+ SVerticalBox::Slot()
+											.AutoHeight()
+											.HAlign(HAlign_Center)
+											.Padding(5.0f)
+											[
+												SNew(STextBlock)
+												.WrapTextAt(150.0f)												
+												.Visibility_Lambda([this]()
+												{
+													return ProjectPresetCollectionsVisibility() == EVisibility::Visible ? EVisibility::Collapsed : EVisibility::Visible;
+												})
+												.Text(LOCTEXT("ProjectPresetsNotLoadedLabel", "Manage Project Preset Collections in Project Settings"))
+												.Justification(ETextJustify::Center)
+												.Font(FAppStyle::GetFontStyle("NormalFontItalic"))
+											]
+
+
+								]
+
+						]
+
+						+ SSplitter::Slot()
+						[
+							SNew(SVerticalBox)
+
+							+ SVerticalBox::Slot()
+							.FillHeight(1.0)
+							[
+								SNew(SOverlay)
+								+ SOverlay::Slot()
+								.ZOrder(1)
+							[
+
+
+								SAssignNew(PresetListView, SListView<TSharedPtr<FToolPresetViewEntry>>)
+								.ListItemsSource(&PresetDataList)
+								.ItemHeight(32.0f)
+								.SelectionMode(ESelectionMode::SingleToggle)
+								.OnGenerateRow(this, &SToolPresetManager::HandleListGenerateRow)		
+								.OnSelectionChanged(this, &SToolPresetManager::HandleListSelectionChanged)
+								.OnContextMenuOpening(this, &SToolPresetManager::OnGetPresetContextMenuContent)
+								.HeaderRow
+								(
+									SNew(SHeaderRow)
+									.Visibility(EVisibility::Visible)
+
+									+ SHeaderRow::Column("Tool")
+									.FixedWidth(30.0f)
+									.HeaderContentPadding(FMargin(5.0f, 5.0f))
+									.HAlignHeader(EHorizontalAlignment::HAlign_Center)
+									.VAlignHeader(EVerticalAlignment::VAlign_Center)
+									.HAlignCell(EHorizontalAlignment::HAlign_Center)
+									.HeaderContent()									
+									[
+										SNew(SImage)
+										.Image(FToolPresetEditorStyle::Get()->GetBrush("ManagerIcons.Tools"))
+										.DesiredSizeOverride(FVector2D(20, 20))
+									]
+									
+									+ SHeaderRow::Column("Label")
+									.FillWidth(80.0f)
+									.HeaderContentPadding(FMargin(5.0f, 5.0f))
+									.HAlignHeader(EHorizontalAlignment::HAlign_Left)
+									.VAlignHeader(EVerticalAlignment::VAlign_Center)
+									.HeaderContent()
+									[
+										SNew(STextBlock)
+										.Text(LOCTEXT("ToolPresetManagerPresetLabelHeader", "Label"))
+									]
+
+									+ SHeaderRow::Column("Tooltip")
+									.FillWidth(80.0f)
+									.HAlignHeader(EHorizontalAlignment::HAlign_Left)
+									.VAlignHeader(EVerticalAlignment::VAlign_Center)
+									.HeaderContent()
+									[
+										SNew(STextBlock)
+										.Text(LOCTEXT("ToolPresetManagerPresetTooltipHeader", "Tooltip"))
+									]
+								)
+							]
+							
+							+ SOverlay::Slot()
+							.HAlign(HAlign_Center)
+							.VAlign(VAlign_Center)
+							.ZOrder(2)
+								[
+									SNew(STextBlock)
+									.WrapTextAt(150.0f)
+									.Visibility_Lambda([this]()
+									{
+										return bHasPresetsInCollection ? EVisibility::Collapsed : EVisibility::Visible;
+									})
+									.Text(LOCTEXT("NoPresetsAvailableLabel", "Add New Presets from any Modeling Tool"))
+									.Justification(ETextJustify::Center)
+									.Font(FAppStyle::GetFontStyle("NormalFontItalic"))
+								]
+						]
+
+						+ SVerticalBox::Slot()
+						.AutoHeight()
+						.Padding(5.0f)
+						[
+							SNew(SVerticalBox)
+							+ SVerticalBox::Slot()
+							.AutoHeight()
+							.Padding(5.0f)
+							[
+								SNew(SHorizontalBox)								
+								+ SHorizontalBox::Slot()
+							    .FillWidth(1.f)
+								.Padding(5.0f)
+								.HAlign(EHorizontalAlignment::HAlign_Left)	
+								[
+									SNew(STextBlock)
+									.Text(LOCTEXT("ToolPresetLabelEditLabel", "Label"))
+								]
+								+ SHorizontalBox::Slot()
+								.FillWidth(1.f)								
+								.Padding(5.0f)
+								[
+									SNew(SEditableTextBox)
+									.IsEnabled(this, &SToolPresetManager::EditAreaEnabled)
+									.OverflowPolicy(ETextOverflowPolicy::Ellipsis)
+									.Text_Lambda([this]() 
+									{
+										if (ActivePresetToEdit)
+										{
+											return FText::FromString(ActivePresetToEdit->PresetLabel);
+										}
+										return FText::GetEmpty();
+									})
+									.OnTextChanged_Lambda([this](const FText& NewText)
+									{
+										if (ActivePresetToEdit)
+										{
+											// Cap the number of characters sent out of the text box, so we don't overflow menus and tooltips
+											ActivePresetToEdit->PresetLabel = NewText.ToString().Left(255);
+										}
+									})
+									.OnTextCommitted_Lambda([this](const FText& NewText, ETextCommit::Type CommitStatus)
+									{
+										if (ActivePresetToEdit)
+										{
+											// Cap the number of characters sent out of the text box, so we don't overflow menus and tooltips
+											SetPresetLabel(ActivePresetToEdit, FText::FromString(NewText.ToString().Left(255)));
+										}
+									})
+								]
+							]
+							+ SVerticalBox::Slot()
+							.AutoHeight()
+							.Padding(5.0f)
+							[
+								SNew(SHorizontalBox)
+								+ SHorizontalBox::Slot()
+								.FillWidth(1.f)
+								.Padding(5.0f)
+								.HAlign(EHorizontalAlignment::HAlign_Left)
+								[
+									SNew(STextBlock)
+									.Text(LOCTEXT("ToolPresetTooltipEditLabel", "Tooltip"))
+								]
+								+ SHorizontalBox::Slot()
+								.FillWidth(1.f)
+								.Padding(5.0f)																
+								[
+									SNew(SBox)
+									.MinDesiredHeight(44.f)
+									.MaxDesiredHeight(44.0f)
+									[
+										SNew(SMultiLineEditableTextBox)		
+										.IsEnabled(this, &SToolPresetManager::EditAreaEnabled)
+										.AllowMultiLine(false)
+										.AutoWrapText(true)
+										.WrappingPolicy(ETextWrappingPolicy::DefaultWrapping)
+										.OverflowPolicy(ETextOverflowPolicy::Ellipsis)
+										.Text_Lambda([this]()
+										{
+											if (ActivePresetToEdit)
+											{
+												return FText::FromString(ActivePresetToEdit->PresetTooltip);
+											}
+											return FText::GetEmpty();
+										})
+										.OnTextChanged_Lambda([this](const FText& NewText) 
+										{
+											if (ActivePresetToEdit)
+											{
+												// Cap the number of characters sent out of the text box, so we don't overflow menus and tooltips
+												ActivePresetToEdit->PresetTooltip = NewText.ToString().Left(2048);
+											}
+										})
+										.OnTextCommitted_Lambda([this](const FText& NewText, ETextCommit::Type CommitStatus)
+										{
+											if (ActivePresetToEdit)
+											{
+												// Cap the number of characters sent out of the text box, so we don't overflow menus and tooltips
+												SetPresetTooltip(ActivePresetToEdit, FText::FromString(NewText.ToString().Left(2048)));
+											}
+										})
+									]
+								]
+							]
+						]						
+					]
+				]
+		];
+
+		RegeneratePresetTrees();
+		if (UserCollectionsDataList.Num() == 0)
+		{
+			bAreUserCollectionsExpanded = false;
+		}
+		if (ProjectCollectionsDataList.Num() == 0)
+		{
+			bAreProjectCollectionsExpanded = false;
+		}
+
+}
+END_SLATE_FUNCTION_BUILD_OPTIMIZATION
+
+SToolPresetManager::~SToolPresetManager()
+{
+}
+
+void SToolPresetManager::Tick(const FGeometry& AllottedGeometry, const double InCurrentTime, const float InDeltaTime)
+{
+	RegeneratePresetTrees();
+}
+
+void SToolPresetManager::RegeneratePresetTrees()
+{
+	if (!ensure(UserSettings.IsValid()))
+	{
+		return;
+	}
+
+	const UToolPresetProjectSettings* ProjectSettings = GetDefault<UToolPresetProjectSettings>();
+	TArray<FSoftObjectPath> AvailablePresetCollections = ProjectSettings->LoadedPresetCollections.Array();
+	TArray<FSoftObjectPath> AvailableUserPresetCollections;
+	UE::ToolPresetEditor::Private::GetObjectsOfClass<UInteractiveToolsPresetCollectionAsset>(AvailableUserPresetCollections);
+
+	TotalPresetCount = 0;
+
+	auto GenerateSubTree = [this](UInteractiveToolsPresetCollectionAsset* PresetCollection, TSharedPtr<FToolPresetViewEntry> RootEntry)
+	{
+		TMap<FString, FInteractiveToolPresetStore >::TIterator ToolNameIter = PresetCollection->PerToolPresets.CreateIterator();
+		for (; (bool)ToolNameIter; ++ToolNameIter)
+		{
+			int32 ToolCount = 0;
+			for (int32 PresetIndex = 0; PresetIndex < ToolNameIter.Value().NamedPresets.Num(); ++PresetIndex)
+			{
+				ToolCount += ToolNameIter.Value().NamedPresets[PresetIndex].IsValid() ? 1 : 0;
+			}
+			if (ToolCount)
+			{
+				RootEntry->Children.Add(MakeShared<FToolPresetViewEntry>(
+					ToolNameIter.Value().ToolLabel,
+					ToolNameIter.Value().ToolIcon,
+					RootEntry->CollectionPath,
+					ToolNameIter.Key(),
+					ToolCount));
+				RootEntry->Children.Last()->Parent = RootEntry;
+				RootEntry->Count += ToolCount;
+				TotalPresetCount += ToolCount;
+			}
+		}
+
+	};
+
+	auto GenerateTreeEntries = [this, &GenerateSubTree](TObjectPtr<UInteractiveToolsPresetCollectionAsset> DefaultCollection,
+		TArray<FSoftObjectPath>* AssetList,
+		TArray< TSharedPtr< FToolPresetViewEntry > >& TreeList,
+		TSharedPtr<STreeView<TSharedPtr<FToolPresetViewEntry> > >& TreeView)
+	{
+		bool bTreeNeedsRefresh = false;
+		TArray< TSharedPtr< FToolPresetViewEntry > > TempTreeDataList;
+
+		if (DefaultCollection)
+		{
+			TSharedPtr<FToolPresetViewEntry> CollectionEntry = MakeShared<FToolPresetViewEntry>(
+				UserSettings->bDefaultCollectionEnabled,
+				FSoftObjectPath(),
+				DefaultCollection->CollectionLabel,
+				0);
+			CollectionEntry->bIsDefaultCollection = true;
+			GenerateSubTree(DefaultCollection, CollectionEntry);
+			TempTreeDataList.Add(CollectionEntry);
+		}
+
+		if (AssetList)
+		{
+			AssetList->RemoveAll([](const FSoftObjectPath& Path)
+			{
+				return !Path.IsAsset();
+			});
+
+			for (const FSoftObjectPath& Path : *AssetList)
+			{
+				UInteractiveToolsPresetCollectionAsset* PresetCollection = nullptr;
+
+				if (Path.IsAsset())
+				{
+					PresetCollection = Cast<UInteractiveToolsPresetCollectionAsset>(Path.TryLoad());
+				}
+				if (PresetCollection)
+				{
+					TSharedPtr<FToolPresetViewEntry> CollectionEntry = MakeShared<FToolPresetViewEntry>(
+						UserSettings->EnabledPresetCollections.Contains(Path),
+						Path,
+						PresetCollection->CollectionLabel,
+						0);
+					GenerateSubTree(PresetCollection, CollectionEntry);
+					TempTreeDataList.Add(CollectionEntry);
+				}
+			}
+		}
+
+		if (TempTreeDataList.Num() != TreeList.Num())
+		{
+			bTreeNeedsRefresh = true;
+		}
+		else
+		{
+			for (int32 CollectionIndex = 0; CollectionIndex < TreeList.Num(); ++CollectionIndex)
+			{
+				if (!(TreeList[CollectionIndex]->HasSameMetadata(*TempTreeDataList[CollectionIndex])))
+				{
+					bTreeNeedsRefresh = true;
+				}
+			}
+		}
+
+		if (bTreeNeedsRefresh)
+		{
+			TreeList = TempTreeDataList;
+			TreeView->RequestTreeRefresh();
+			bHasActiveCollection = false;
+		}
+
+		for (TSharedPtr<FToolPresetViewEntry>& Entry : TreeList)
+		{
+			Entry->bEnabled = UserSettings->EnabledPresetCollections.Contains(Entry->CollectionPath);
+			if (Entry->bIsDefaultCollection)
+			{
+				Entry->bEnabled = UserSettings->bDefaultCollectionEnabled;
+			}
+			else
+			{
+				Entry->bEnabled = UserSettings->EnabledPresetCollections.Contains(Entry->CollectionPath);
+			}
+		}
+	};
+
+	// Handle the default collection
+	UToolPresetAssetSubsystem* PresetAssetSubsystem = GEditor->GetEditorSubsystem<UToolPresetAssetSubsystem>();
+	TObjectPtr<UInteractiveToolsPresetCollectionAsset> DefaultCollection = nullptr;
+	if (ensure(PresetAssetSubsystem))
+	{
+		DefaultCollection = PresetAssetSubsystem->GetDefaultCollection();
+	}
+
+	GenerateTreeEntries(nullptr, &AvailablePresetCollections, ProjectCollectionsDataList, ProjectPresetCollectionTreeView);
+	GenerateTreeEntries(nullptr, &AvailableUserPresetCollections, UserCollectionsDataList, UserPresetCollectionTreeView);
+	GenerateTreeEntries(DefaultCollection, nullptr, EditorCollectionsDataList, EditorPresetCollectionTreeView);
+
+}
+
+
+/* SToolPresetManager implementation
+ *****************************************************************************/
+
+int32 SToolPresetManager::GetTotalPresetCount() const
+{
+	return TotalPresetCount;
+}
+
+TSharedRef<ITableRow> SToolPresetManager::HandleTreeGenerateRow(TSharedPtr<FToolPresetViewEntry> TreeEntry, const TSharedRef<STableViewBase>& OwnerTable)
+{
+	return SNew(UE::ToolPresetEditor::Private::SCollectionTableRow< TSharedPtr<FToolPresetViewEntry> >, OwnerTable)
+		.ViewEntry(TreeEntry)
+		.OnCollectionEnabledCheckboxChanged(this, &SToolPresetManager::SetCollectionEnabled)
+		.OnCollectionRenameStarted(this, &SToolPresetManager::CollectionRenameStarted)
+		.OnCollectionRenameEnded(this, &SToolPresetManager::CollectionRenameEnded);
+}
+
+void SToolPresetManager::HandleTreeGetChildren(TSharedPtr<FToolPresetViewEntry> TreeEntry, TArray< TSharedPtr<FToolPresetViewEntry> >& ChildrenOut)
+{
+	ChildrenOut = TreeEntry->Children;
+}
+
+void SToolPresetManager::GeneratePresetList(TSharedPtr<FToolPresetViewEntry> TreeEntry)
+{
+	PresetDataList.Empty();
+	PresetListView->RequestListRefresh();
+	bHasActiveCollection = false;
+	ActivePresetToEdit = nullptr;
+	bHasPresetsInCollection = false;
+
+	if (!TreeEntry)
+	{
+		return;
+	}
+
+	if (TreeEntry->EntryType == FToolPresetViewEntry::EEntryType::Collection ||
+		TreeEntry->EntryType == FToolPresetViewEntry::EEntryType::Tool)
+	{
+		UInteractiveToolsPresetCollectionAsset* PresetCollection = GetCollectionFromEntry(TreeEntry);
+
+		if (PresetCollection)
+		{
+			if (TreeEntry->EntryType == FToolPresetViewEntry::EEntryType::Collection)
+			{
+				bHasActiveCollection = true;
+				bIsActiveCollectionEnabled = TreeEntry->bEnabled;
+				ActiveCollectionLabel = TreeEntry->EntryLabel;
+
+				TMap<FString, FInteractiveToolPresetStore >::TIterator ToolNameIter = PresetCollection->PerToolPresets.CreateIterator();
+				for (; (bool)ToolNameIter; ++ToolNameIter)
+				{
+					int32 ToolCount = ToolNameIter.Value().NamedPresets.Num();
+					for (int32 PresetIndex = 0; PresetIndex < ToolNameIter.Value().NamedPresets.Num(); ++PresetIndex)
+					{
+						if (ToolNameIter.Value().NamedPresets[PresetIndex].IsValid())
+						{
+							bHasPresetsInCollection = true;
+							PresetDataList.Add(MakeShared<FToolPresetViewEntry>(
+								ToolNameIter.Key(),
+								PresetIndex,
+								ToolNameIter.Value().NamedPresets[PresetIndex].Label,
+								ToolNameIter.Value().NamedPresets[PresetIndex].Tooltip,
+								FText::FromString(ToolNameIter.Value().NamedPresets[PresetIndex].Label)
+								));
+							PresetDataList.Last()->Parent = TreeEntry;
+							PresetDataList.Last()->CollectionPath = TreeEntry->CollectionPath;
+							PresetDataList.Last()->EntryIcon = ToolNameIter.Value().ToolIcon;
+						}
+					}
+				}
+			}
+			else
+			{
+				bHasActiveCollection = true;
+				bIsActiveCollectionEnabled = TreeEntry->Parent->bEnabled;
+				ActiveCollectionLabel = TreeEntry->Parent->EntryLabel;
+
+				const FInteractiveToolPresetStore* ToolData = PresetCollection->PerToolPresets.Find(TreeEntry->ToolName);
+				if (!ToolData)
+				{
+					return;
+				}
+				int32 ToolCount = ToolData->NamedPresets.Num();
+				for (int32 PresetIndex = 0; PresetIndex < ToolData->NamedPresets.Num(); ++PresetIndex)
+				{
+					if (ToolData->NamedPresets[PresetIndex].IsValid())
+					{
+						bHasPresetsInCollection = true;
+						PresetDataList.Add(MakeShared<FToolPresetViewEntry>(
+							TreeEntry->ToolName,
+							PresetIndex,
+							ToolData->NamedPresets[PresetIndex].Label,
+							ToolData->NamedPresets[PresetIndex].Tooltip,
+							FText::FromString(ToolData->NamedPresets[PresetIndex].Label)
+							));
+						PresetDataList.Last()->Parent = TreeEntry;
+						PresetDataList.Last()->CollectionPath = TreeEntry->CollectionPath;
+						PresetDataList.Last()->EntryIcon = TreeEntry->EntryIcon;
+					}
+				}
+			}
+		}
+	}
+
+}
+
+void SToolPresetManager::HandleEditorTreeSelectionChanged(TSharedPtr<FToolPresetViewEntry> TreeEntry, ESelectInfo::Type SelectInfo)
+{
+	for (TSharedPtr<FToolPresetViewEntry> Entry : UserPresetCollectionTreeView->GetRootItems())
+	{
+		Entry->bIsRenaming = false;
+	}
+
+	if (SelectInfo != ESelectInfo::Direct)
+	{
+		UserPresetCollectionTreeView->ClearSelection();
+		ProjectPresetCollectionTreeView->ClearSelection();
+		GeneratePresetList(TreeEntry);
+
+		LastFocusedList = EditorPresetCollectionTreeView;
+	}
+}
+
+void SToolPresetManager::HandleTreeSelectionChanged(TSharedPtr<FToolPresetViewEntry> TreeEntry, ESelectInfo::Type SelectInfo)
+{
+	for (TSharedPtr<FToolPresetViewEntry> Entry : UserPresetCollectionTreeView->GetRootItems())
+	{
+		Entry->bIsRenaming = false;
+	}
+
+	if (SelectInfo != ESelectInfo::Direct)
+	{
+		UserPresetCollectionTreeView->ClearSelection();
+		EditorPresetCollectionTreeView->ClearSelection();
+		GeneratePresetList(TreeEntry);
+
+		LastFocusedList = ProjectPresetCollectionTreeView;
+
+	}
+}
+
+void SToolPresetManager::HandleUserTreeSelectionChanged(TSharedPtr<FToolPresetViewEntry> TreeEntry, ESelectInfo::Type SelectInfo)
+{
+	for (TSharedPtr<FToolPresetViewEntry> Entry : UserPresetCollectionTreeView->GetRootItems())
+	{
+		Entry->bIsRenaming = false;
+	}
+
+	if (SelectInfo != ESelectInfo::Direct)
+	{
+		ProjectPresetCollectionTreeView->ClearSelection();
+		EditorPresetCollectionTreeView->ClearSelection();
+		GeneratePresetList(TreeEntry);
+
+		LastFocusedList = UserPresetCollectionTreeView;
+
+	}
+}
+
+TSharedRef<ITableRow> SToolPresetManager::HandleListGenerateRow(TSharedPtr<FToolPresetViewEntry> TreeEntry, const TSharedRef<STableViewBase>& OwnerTable)
+{
+	return SNew(UE::ToolPresetEditor::Private::SToolPresetTableRow< TSharedPtr<FToolPresetViewEntry> >, OwnerTable)
+		.ViewEntry(TreeEntry)
+		.OnPresetDeleted(this, &SToolPresetManager::DeletePresetFromCollection);
+}
+
+void SToolPresetManager::HandleListSelectionChanged(TSharedPtr<FToolPresetViewEntry> TreeEntry, ESelectInfo::Type SelectInfo)
+{
+	if (SelectInfo != ESelectInfo::Direct)
+	{
+		LastFocusedList = PresetListView;
+	}
+
+
+	if (TreeEntry)
+	{
+		ActivePresetToEdit = TreeEntry;
+	}
+	else
+	{
+		if (ActivePresetToEdit)
+		{
+			SetPresetLabel(ActivePresetToEdit, FText::FromString(ActivePresetToEdit->PresetLabel));
+			SetPresetTooltip(ActivePresetToEdit, FText::FromString(ActivePresetToEdit->PresetTooltip));
+		}
+		ActivePresetToEdit.Reset();
+	}
+}
+
+bool SToolPresetManager::EditAreaEnabled() const
+{
+	return ActivePresetToEdit.IsValid();
+}
+
+EVisibility SToolPresetManager::ProjectPresetCollectionsVisibility() const
+{
+	return ProjectCollectionsDataList.IsEmpty() ? EVisibility::Collapsed : EVisibility::Visible;
+}
+
+void SToolPresetManager::SetCollectionEnabled(TSharedPtr<FToolPresetViewEntry> TreeEntry, ECheckBoxState State)
+{	
+	if (!ensure(UserSettings.IsValid()))
+	{
+		return;
+	}
+	if (TreeEntry->bIsDefaultCollection)
+	{
+		UserSettings->bDefaultCollectionEnabled = (State == ECheckBoxState::Checked);
+		UserSettings->SaveEditorConfig();
+	}
+	else
+	{
+		if (State == ECheckBoxState::Checked && !UserSettings->EnabledPresetCollections.Contains(TreeEntry->CollectionPath))
+		{
+			UserSettings->EnabledPresetCollections.Add(TreeEntry->CollectionPath);
+			UserSettings->SaveEditorConfig();
+		}
+		else if (State != ECheckBoxState::Checked && UserSettings->EnabledPresetCollections.Contains(TreeEntry->CollectionPath))
+		{
+			UserSettings->EnabledPresetCollections.Remove(TreeEntry->CollectionPath);
+			UserSettings->SaveEditorConfig();
+		}
+	}
+}
+
+void SToolPresetManager::CollectionRenameStarted(TSharedPtr<FToolPresetViewEntry> TreeEntry, TSharedPtr<SEditableTextBox> RenameWidget)
+{
+	// TODO: Figure out why this crashes
+	//FSlateApplication::Get().SetKeyboardFocus(RenameWidget, EFocusCause::SetDirectly);
+}
+
+void SToolPresetManager::CollectionRenameEnded(TSharedPtr<FToolPresetViewEntry> TreeEntry, const FText& NewText)
+{
+	FAssetRegistryModule& AssetRegistryModule = FModuleManager::LoadModuleChecked<FAssetRegistryModule>(TEXT("AssetRegistry"));
+
+	TreeEntry->bIsRenaming = false;
+
+	FAssetData CollectionAsset;
+	if (AssetRegistryModule.Get().TryGetAssetByObjectPath(TreeEntry->CollectionPath, CollectionAsset) == UE::AssetRegistry::EExists::Exists)
+	{
+		UInteractiveToolsPresetCollectionAsset* CollectionObject = Cast<UInteractiveToolsPresetCollectionAsset>(CollectionAsset.GetAsset());
+		if (!CollectionObject)
+		{
+			return;
+		}
+
+		TArray< FAssetRenameData> RenameData;
+
+		FString NewPackageName, NewAssetName;
+		IAssetTools::Get().CreateUniqueAssetName(NewText.ToString(), "", NewPackageName, NewAssetName);
+
+		RenameData.SetNum(1);		
+		RenameData[0].Asset = CollectionObject;
+		RenameData[0].NewName = NewAssetName;
+		RenameData[0].NewPackagePath = CollectionAsset.PackagePath.ToString();
+		if(IAssetTools::Get().RenameAssets(RenameData))
+		{
+			if (UserSettings->EnabledPresetCollections.Contains(TreeEntry->CollectionPath))
+			{
+				UserSettings->EnabledPresetCollections.Remove(TreeEntry->CollectionPath);
+				UserSettings->EnabledPresetCollections.Add(CollectionObject->GetPathName());
+				UserSettings->SaveEditorConfig();
+			}
+
+			CollectionObject->CollectionLabel = NewText;
+			CollectionObject->MarkPackageDirty();	
+		}
+	}
+}
+
+
+void SToolPresetManager::DeletePresetFromCollection(TSharedPtr< FToolPresetViewEntry > Entry)
+{
+	UInteractiveToolsPresetCollectionAsset* PresetCollection = GetCollectionFromEntry(Entry);
+	if (PresetCollection)
+	{
+		PresetCollection->PerToolPresets[Entry->ToolName].NamedPresets.RemoveAt(Entry->PresetIndex);
+		PresetCollection->MarkPackageDirty();
+
+		GeneratePresetList(Entry->Parent);
+	}
+
+	SaveIfDefaultCollection(Entry);
+}
+
+void SToolPresetManager::SetPresetLabel(TSharedPtr< FToolPresetViewEntry > Entry, FText InLabel)
+{
+	UInteractiveToolsPresetCollectionAsset* PresetCollection = GetCollectionFromEntry(Entry);
+	if (PresetCollection)
+	{
+		PresetCollection->PerToolPresets[Entry->ToolName].NamedPresets[Entry->PresetIndex].Label = InLabel.ToString();
+		PresetCollection->MarkPackageDirty();
+	}
+
+	SaveIfDefaultCollection(Entry);
+}
+
+void SToolPresetManager::SetPresetTooltip(TSharedPtr< FToolPresetViewEntry > Entry, FText InTooltip)
+{
+	UInteractiveToolsPresetCollectionAsset* PresetCollection = GetCollectionFromEntry(Entry);
+	if (PresetCollection)
+	{
+		PresetCollection->PerToolPresets[Entry->ToolName].NamedPresets[Entry->PresetIndex].Tooltip = InTooltip.ToString();
+		PresetCollection->MarkPackageDirty();
+	}
+
+	SaveIfDefaultCollection(Entry);
+}
+
+void SToolPresetManager::DeleteSelectedUserPresetCollection()
+{
+	FAssetRegistryModule& AssetRegistryModule = FModuleManager::LoadModuleChecked<FAssetRegistryModule>(TEXT("AssetRegistry"));
+	FContentBrowserModule& ContentBrowserModule = FModuleManager::LoadModuleChecked<FContentBrowserModule>("ContentBrowser");
+
+	TArray<TSharedPtr<FToolPresetViewEntry>> SelectedUserCollections = UserPresetCollectionTreeView->GetSelectedItems();
+
+	if (SelectedUserCollections.Num() == 1)
+	{
+		TSharedPtr<FToolPresetViewEntry> Entry = SelectedUserCollections[0];
+		if (Entry->bIsDefaultCollection)
+		{
+			return;
+		}
+
+		FAssetData CollectionAsset;
+		if (AssetRegistryModule.Get().TryGetAssetByObjectPath(Entry->CollectionPath, CollectionAsset) == UE::AssetRegistry::EExists::Exists)
+		{
+			TArray<FAssetData> AssetData;
+			AssetData.Add(CollectionAsset);
+			ObjectTools::DeleteAssets(AssetData, true);
+		}
+
+		GeneratePresetList(nullptr);
+	}
+}
+
+void SToolPresetManager::AddNewUserPresetCollection()
+{
+	// Load necessary modules
+	FAssetToolsModule& AssetToolsModule = FModuleManager::Get().LoadModuleChecked<FAssetToolsModule>("AssetTools");
+	FContentBrowserModule& ContentBrowserModule = FModuleManager::LoadModuleChecked<FContentBrowserModule>("ContentBrowser");
+	FAssetRegistryModule& AssetRegistryModule = FModuleManager::LoadModuleChecked<FAssetRegistryModule>(TEXT("AssetRegistry"));
+	IAssetRegistry& AssetRegistry = AssetRegistryModule.Get();
+
+	// Generate a unique asset name
+	FString Name, PackageName;
+	AssetToolsModule.Get().CreateUniqueAssetName(TEXT("/ToolPresets/Presets/"), TEXT("UserPresetCollection"), PackageName, Name);
+	const FString PackagePath = FPackageName::GetLongPackagePath(PackageName);
+
+	// Create object and package
+	UPackage* package = CreatePackage(*PackageName);
+	UInteractiveToolsPresetCollectionAssetFactory* MyFactory = NewObject<UInteractiveToolsPresetCollectionAssetFactory>(UInteractiveToolsPresetCollectionAssetFactory::StaticClass()); // Can omit, and a default factory will be used
+	UObject* NewObject = AssetToolsModule.Get().CreateAsset(Name, PackagePath, UInteractiveToolsPresetCollectionAsset::StaticClass(), MyFactory);
+	UInteractiveToolsPresetCollectionAsset* NewCollection = ExactCast<UInteractiveToolsPresetCollectionAsset>(NewObject);
+	NewCollection->CollectionLabel = FText::FromString(Name);
+	FSavePackageArgs SavePackageArgs;
+	SavePackageArgs.TopLevelFlags = RF_Public | RF_Standalone;
+	UPackage::Save(package, NewObject, *FPackageName::LongPackageNameToFilename(PackageName, FPackageName::GetAssetPackageExtension()), SavePackageArgs);
+
+	// Inform asset registry
+	AssetRegistry.AssetCreated(NewObject);
+
+	// Since we're adding a new entry, open this tree view up again.
+	bAreUserCollectionsExpanded = true;
+}
+
+const FSlateBrush* SToolPresetManager::GetProjectCollectionsExpanderImage() const
+{
+	return GetExpanderImage(ProjectCollectionsExpander, false);
+}
+
+const FSlateBrush* SToolPresetManager::GetUserCollectionsExpanderImage() const
+{
+	return GetExpanderImage(UserCollectionsExpander, true);
+}
+
+const FSlateBrush* SToolPresetManager::GetExpanderImage(TSharedPtr<SWidget> ExpanderWidget, bool bIsUserCollections) const
+{
+	const bool bIsItemExpanded = bIsUserCollections ? bAreUserCollectionsExpanded : bAreProjectCollectionsExpanded;
+
+	FName ResourceName;
+	if (bIsItemExpanded)
+	{
+		if (ExpanderWidget->IsHovered())
+		{
+			static FName ExpandedHoveredName = "TreeArrow_Expanded_Hovered";
+			ResourceName = ExpandedHoveredName;
+		}
+		else
+		{
+			static FName ExpandedName = "TreeArrow_Expanded";
+			ResourceName = ExpandedName;
+		}
+	}
+	else
+	{
+		if (ExpanderWidget->IsHovered())
+		{
+			static FName CollapsedHoveredName = "TreeArrow_Collapsed_Hovered";
+			ResourceName = CollapsedHoveredName;
+		}
+		else
+		{
+			static FName CollapsedName = "TreeArrow_Collapsed";
+			ResourceName = CollapsedName;
+		}
+	}
+
+	return FCoreStyle::Get().GetBrush(ResourceName);
+}
+
+UInteractiveToolsPresetCollectionAsset* SToolPresetManager::GetCollectionFromEntry(TSharedPtr<FToolPresetViewEntry> Entry)
+{
+	UInteractiveToolsPresetCollectionAsset* PresetCollection = nullptr;
+	UToolPresetAssetSubsystem* PresetAssetSubsystem = GEditor->GetEditorSubsystem<UToolPresetAssetSubsystem>();
+	
+	if (Entry->Root().bIsDefaultCollection && ensure(PresetAssetSubsystem))
+	{
+		PresetCollection = PresetAssetSubsystem->GetDefaultCollection();
+	}
+	else
+	{
+		if (Entry->CollectionPath.IsAsset())
+		{
+			PresetCollection = Cast<UInteractiveToolsPresetCollectionAsset>(Entry->CollectionPath.TryLoad());
+		}
+	}
+
+	return PresetCollection;
+}
+
+void SToolPresetManager::SaveIfDefaultCollection(TSharedPtr<FToolPresetViewEntry> Entry)
+{
+	UToolPresetAssetSubsystem* PresetAssetSubsystem = GEditor->GetEditorSubsystem<UToolPresetAssetSubsystem>();
+
+	if (Entry->Root().bIsDefaultCollection && ensure(PresetAssetSubsystem))
+	{
+		ensure(PresetAssetSubsystem->SaveDefaultCollection());
+	}
+}
+
+TSharedPtr<SWidget> SToolPresetManager::OnGetPresetContextMenuContent() const
+{
+	const bool bShouldCloseWindowAfterMenuSelection = true;
+	FMenuBuilder MenuBuilder(bShouldCloseWindowAfterMenuSelection, UICommandList);
+
+	MenuBuilder.BeginSection("ToolPresetManagerPresetAction", LOCTEXT("ToolPresetAction", "Tool Preset Actions"));
+
+	MenuBuilder.AddMenuEntry(FGenericCommands::Get().Delete, NAME_None, LOCTEXT("DeleteToolPresetLabel", "Delete Tool Preset"), LOCTEXT("DeleteToolPresetToolTip", "Delete the selected tool preset"));
+
+	MenuBuilder.EndSection();
+
+	return MenuBuilder.MakeWidget();
+}
+
+TSharedPtr<SWidget> SToolPresetManager::OnGetCollectionContextMenuContent() const
+{
+	const bool bShouldCloseWindowAfterMenuSelection = true;
+	FMenuBuilder MenuBuilder(bShouldCloseWindowAfterMenuSelection, UICommandList);
+
+	MenuBuilder.BeginSection("ToolPresetManagerCollectionAction", LOCTEXT("CollectionAction", "Tool Preset Collection Actions"));
+
+	MenuBuilder.AddMenuEntry(FGenericCommands::Get().Delete, NAME_None, LOCTEXT("DeleteCollectionLabel", "Delete Collection"), LOCTEXT("DeleteCollectionToolTip", "Delete the selected collection"));
+	MenuBuilder.AddMenuEntry(FGenericCommands::Get().Rename, NAME_None, LOCTEXT("RenameCollectionLabel", "Rename Collection"), LOCTEXT("RenameCollectionToolTip", "Rename the selected collection"));
+
+	MenuBuilder.EndSection();
+
+	return MenuBuilder.MakeWidget();
+}
+
+void SToolPresetManager::BindCommands()
+{
+	// This should not be called twice on the same instance
+	check(!UICommandList.IsValid());
+
+	UICommandList = MakeShareable(new FUICommandList);
+
+	FUICommandList& CommandList = *UICommandList;
+
+	// ...and bind them all
+
+	CommandList.MapAction(
+		FGenericCommands::Get().Delete,
+		FExecuteAction::CreateSP(this, &SToolPresetManager::OnDeleteClicked),
+		FCanExecuteAction::CreateSP(this, &SToolPresetManager::CanDelete));
+
+	CommandList.MapAction(
+		FGenericCommands::Get().Rename,
+		FExecuteAction::CreateSP(this, &SToolPresetManager::OnRenameClicked),
+		FCanExecuteAction::CreateSP(this, &SToolPresetManager::CanRename));	
+}
+
+
+void SToolPresetManager::OnDeleteClicked()
+{
+	if (UserPresetCollectionTreeView == LastFocusedList)
+	{
+		DeleteSelectedUserPresetCollection();
+	}
+
+	if (PresetListView == LastFocusedList)
+	{
+		for (TSharedPtr<FToolPresetViewEntry> Entry : PresetListView->GetSelectedItems())
+		{
+			DeletePresetFromCollection(Entry);
+		}
+	}
+}
+
+bool SToolPresetManager::CanDelete()
+{
+	bool bIsListValid = LastFocusedList.IsValid() && (LastFocusedList == UserPresetCollectionTreeView || LastFocusedList == PresetListView);
+	bool bIsSelectionValid = false;
+	if (bIsListValid)
+	{
+		bIsSelectionValid = LastFocusedList.Pin()->GetNumItemsSelected() == 1 &&
+			                (LastFocusedList.Pin()->GetSelectedItems()[0]->EntryType == FToolPresetViewEntry::EEntryType::Collection ||
+							 LastFocusedList.Pin()->GetSelectedItems()[0]->EntryType == FToolPresetViewEntry::EEntryType::Preset);
+	}
+	return bIsSelectionValid;
+}
+
+void SToolPresetManager::OnRenameClicked()
+{
+	for (TSharedPtr<FToolPresetViewEntry> Entry : UserPresetCollectionTreeView->GetRootItems())
+	{
+		Entry->bIsRenaming = false;
+	}
+
+	for (TSharedPtr<FToolPresetViewEntry> Entry : UserPresetCollectionTreeView->GetSelectedItems())
+	{
+		Entry->bIsRenaming = true;
+	}
+}
+
+bool SToolPresetManager::CanRename()
+{
+	bool bIsListValid = LastFocusedList.IsValid() && LastFocusedList == UserPresetCollectionTreeView;
+	bool bIsSelectionValid = false;
+	if (bIsListValid)
+	{
+		bIsSelectionValid = LastFocusedList.Pin()->GetNumItemsSelected() == 1 &&
+			                LastFocusedList.Pin()->GetSelectedItems()[0]->EntryType == FToolPresetViewEntry::EEntryType::Collection;
+	}
+	return bIsSelectionValid;
+}
+
+#undef LOCTEXT_NAMESPACE
diff --git a/Engine/Plugins/Experimental/ToolPresets/Source/ToolPresetEditor/Private/Widgets/SToolPresetManager.h b/Engine/Plugins/Experimental/ToolPresets/Source/ToolPresetEditor/Private/Widgets/SToolPresetManager.h
new file mode 100644
--- /dev/null
+++ b/Engine/Plugins/Experimental/ToolPresets/Source/ToolPresetEditor/Private/Widgets/SToolPresetManager.h
@@ -0,0 +1,254 @@
+// Copyright Epic Games, Inc. All Rights Reserved.
+
+#pragma once
+
+#include "CoreMinimal.h"
+#include "Input/Reply.h"
+#include "Widgets/DeclarativeSyntaxSupport.h"
+#include "Widgets/SCompoundWidget.h"
+#include "Widgets/Views/STableRow.h"
+#include "Widgets/Views/STableViewBase.h"
+
+class FTransaction;
+class FUICommandList;
+class SEditableTextBox;
+class SNegativeActionButton;
+class SPositiveActionButton;
+class SSplitter;
+class SVerticalBox;
+class UInteractiveToolsPresetCollectionAsset;
+class UToolPresetUserSettings;
+template<class ItemType> class SListView;
+template<class ItemType> class STreeView;
+
+/**
+ * Implements the preset manager panel.
+ */
+class SToolPresetManager
+	: public SCompoundWidget
+{
+	struct FToolPresetCollectionInfo
+	{
+		FSoftObjectPath PresetCollectionPath;
+		bool bCollectionEnabled;
+
+		FToolPresetCollectionInfo(FSoftObjectPath InPresetCollectionPath, bool bInEnabled)
+			: PresetCollectionPath(InPresetCollectionPath), bCollectionEnabled(bInEnabled)
+		{}
+	};
+
+public:
+	struct FToolPresetViewEntry
+	{
+		enum class EEntryType : uint8
+		{
+			Collection,
+			Tool,
+			Preset
+		};
+		
+		EEntryType EntryType;
+
+		// Used for Collections/Tool Entries
+		bool bEnabled = false;
+		FSoftObjectPath CollectionPath;
+		bool bIsDefaultCollection = false;
+		bool bIsRenaming = false;
+		FText EntryLabel;
+		FSlateBrush EntryIcon;
+		int32 Count = 0;
+
+		// Used for Preset/Tool entries
+		FString ToolName = "";
+		int32 PresetIndex = 0;
+		FString PresetLabel = "";
+		FString PresetTooltip = "";
+
+		TSharedPtr< FToolPresetViewEntry> Parent;
+		TArray<TSharedPtr< FToolPresetViewEntry> > Children;
+
+		// Collection Constructor
+		FToolPresetViewEntry(bool bEnabledIn, FSoftObjectPath CollectionPathIn, FText EntryLabelIn, int32 CountIn)
+			: EntryType(EEntryType::Collection),
+			  bEnabled(bEnabledIn), 
+			  CollectionPath(CollectionPathIn),
+			  EntryLabel(EntryLabelIn),
+			  Count(CountIn)
+		{}
+
+		// Tool Constructor
+		FToolPresetViewEntry(FText EntryLabelIn, FSlateBrush EntryIconIn, FSoftObjectPath CollectionPathIn, FString ToolNameIn, int32 CountIn)
+		  :	EntryType(EEntryType::Tool),
+			CollectionPath(CollectionPathIn),
+			EntryLabel(EntryLabelIn),
+			EntryIcon(EntryIconIn),
+			Count(CountIn),
+			ToolName(ToolNameIn)
+		{}
+
+		// Preset Constructor
+		FToolPresetViewEntry(FString ToolNameIn, int32 PresetIndexIn, FString PresetLabelIn, FString PresetTooltipIn, FText EntryLabelIn)
+			: EntryType(EEntryType::Preset),
+			  EntryLabel(EntryLabelIn),
+			  ToolName(ToolNameIn),
+			  PresetIndex(PresetIndexIn),
+			  PresetLabel(PresetLabelIn),
+			  PresetTooltip(PresetTooltipIn)
+		{}
+
+		bool HasSameMetadata(FToolPresetViewEntry& Other)
+		{
+			bool bIsEqual = 
+				 EntryType == Other.EntryType &&
+				 CollectionPath == Other.CollectionPath &&
+				 bIsDefaultCollection == Other.bIsDefaultCollection &&
+				 Count == Other.Count &&
+				 ToolName.Equals(Other.ToolName) &&
+				 PresetIndex == Other.PresetIndex &&
+				 PresetLabel.Equals(Other.PresetLabel) &&
+				 PresetTooltip.Equals(Other.PresetTooltip) &&
+				 Children.Num() == Other.Children.Num();
+
+			for (int32 ChildIndex = 0; ChildIndex < Children.Num(); ++ChildIndex)
+			{
+				bIsEqual = bIsEqual && (*Children[ChildIndex]).HasSameMetadata(*Other.Children[ChildIndex]);
+			}
+			return bIsEqual;
+		}
+
+		bool operator==(FToolPresetViewEntry& Other)
+		{
+			bool bIsEqual = bEnabled == Other.bEnabled &&
+				 EntryType == Other.EntryType &&
+				 CollectionPath == Other.CollectionPath &&
+				 bIsDefaultCollection == Other.bIsDefaultCollection && 
+				 Count == Other.Count &&
+				 EntryLabel.EqualTo(Other.EntryLabel) &&
+				 ToolName.Equals(Other.ToolName) &&
+				 PresetIndex == Other.PresetIndex &&
+				 PresetLabel.Equals(Other.PresetLabel) &&
+				 PresetTooltip.Equals(Other.PresetTooltip) &&
+				 Children.Num() == Other.Children.Num();
+
+			for (int32 ChildIndex = 0; ChildIndex < Children.Num(); ++ChildIndex)
+			{
+				bIsEqual = bIsEqual && (*Children[ChildIndex]) == (*Other.Children[ChildIndex]);
+			}
+
+			return bIsEqual;
+		}
+
+		FToolPresetViewEntry& Root()
+		{
+			FToolPresetViewEntry* ActiveNode = this;
+			while (ActiveNode->Parent)
+			{
+				ActiveNode = ActiveNode->Parent.Get();
+			}
+			return *ActiveNode;
+		}
+	};
+
+	SLATE_BEGIN_ARGS(SToolPresetManager) { }
+	SLATE_END_ARGS()
+
+	virtual ~SToolPresetManager();
+
+	/**
+	 * Construct this widget
+	 *
+	 * @param InArgs The declaration data for this widget.
+	 */
+	void Construct( const FArguments& InArgs );
+	
+	//~ Begin SWidget Interface
+	virtual void Tick( const FGeometry& AllottedGeometry, const double InCurrentTime, const float InDeltaTime ) override;
+	//~ End SWidget Interface
+
+private:
+	void BindCommands();
+
+	void RegeneratePresetTrees();
+
+	int32 GetTotalPresetCount() const;
+
+	TSharedRef<ITableRow> HandleTreeGenerateRow(TSharedPtr<FToolPresetViewEntry> TreeEntry, const TSharedRef<STableViewBase>& OwnerTable);
+	void HandleTreeGetChildren(TSharedPtr<FToolPresetViewEntry> TreeEntry, TArray< TSharedPtr<FToolPresetViewEntry> >& ChildrenOut);
+	void HandleTreeSelectionChanged(TSharedPtr<FToolPresetViewEntry> TreeEntry, ESelectInfo::Type);
+	void HandleUserTreeSelectionChanged(TSharedPtr<FToolPresetViewEntry> TreeEntry, ESelectInfo::Type);
+	void HandleEditorTreeSelectionChanged(TSharedPtr<FToolPresetViewEntry> TreeEntry, ESelectInfo::Type);
+	EVisibility ProjectPresetCollectionsVisibility() const;
+
+	TSharedPtr<SWidget> OnGetPresetContextMenuContent() const;
+	TSharedPtr<SWidget> OnGetCollectionContextMenuContent() const;
+
+	void GeneratePresetList(TSharedPtr<FToolPresetViewEntry> TreeEntry);
+	TSharedRef<ITableRow> HandleListGenerateRow(TSharedPtr<FToolPresetViewEntry> TreeEntry, const TSharedRef<STableViewBase>& OwnerTable);
+    void HandleListSelectionChanged(TSharedPtr<FToolPresetViewEntry> TreeEntry, ESelectInfo::Type SelectInfo);
+
+	bool EditAreaEnabled() const;
+	void SetCollectionEnabled(TSharedPtr<FToolPresetViewEntry> TreeEntry, ECheckBoxState State);
+	void DeletePresetFromCollection(TSharedPtr< FToolPresetViewEntry > Entry);
+
+	void CollectionRenameStarted(TSharedPtr<FToolPresetViewEntry> TreeEntry, TSharedPtr<SEditableTextBox> RenameWidget);
+	void CollectionRenameEnded(TSharedPtr<FToolPresetViewEntry> TreeEntry, const FText& NewText);
+
+	void DeleteSelectedUserPresetCollection();
+	void AddNewUserPresetCollection();
+
+	void SetPresetLabel(TSharedPtr< FToolPresetViewEntry >, FText InLabel);
+	void SetPresetTooltip(TSharedPtr< FToolPresetViewEntry >, FText InTooltip);
+
+	const FSlateBrush* GetProjectCollectionsExpanderImage() const;
+	const FSlateBrush* GetUserCollectionsExpanderImage() const;
+	const FSlateBrush* GetExpanderImage(TSharedPtr<SWidget> ExpanderWidget, bool bIsUserCollections) const;
+
+	UInteractiveToolsPresetCollectionAsset* GetCollectionFromEntry(TSharedPtr<FToolPresetViewEntry> Entry);
+	void SaveIfDefaultCollection(TSharedPtr<FToolPresetViewEntry> Entry);
+
+	void OnDeleteClicked();
+	bool CanDelete();
+
+	void OnRenameClicked();
+	bool CanRename();
+
+private:	
+	TSharedPtr<FUICommandList> UICommandList;
+
+	TWeakObjectPtr<UToolPresetUserSettings> UserSettings;
+
+	TWeakPtr< SListView<TSharedPtr<FToolPresetViewEntry> > >  LastFocusedList;
+
+	bool bAreProjectCollectionsExpanded = true;
+	TSharedPtr<SButton> ProjectCollectionsExpander;
+	TArray< TSharedPtr< FToolPresetViewEntry > > ProjectCollectionsDataList;
+	TSharedPtr<STreeView<TSharedPtr<FToolPresetViewEntry> > > ProjectPresetCollectionTreeView;
+
+	bool bAreUserCollectionsExpanded = true;
+	TSharedPtr<SButton> UserCollectionsExpander;
+	TArray< TSharedPtr< FToolPresetViewEntry > > UserCollectionsDataList;
+	TSharedPtr<STreeView<TSharedPtr<FToolPresetViewEntry> > > UserPresetCollectionTreeView;
+
+	TArray< TSharedPtr< FToolPresetViewEntry > > EditorCollectionsDataList;
+	TSharedPtr<STreeView<TSharedPtr<FToolPresetViewEntry> > > EditorPresetCollectionTreeView;
+
+
+	TArray< TSharedPtr< FToolPresetViewEntry > > PresetDataList;
+	TSharedPtr<SListView<TSharedPtr<FToolPresetViewEntry> > > PresetListView;
+
+	int32 TotalPresetCount = 0;
+	bool bHasActiveCollection = false;
+	bool bHasPresetsInCollection = false;
+	FText ActiveCollectionLabel;
+	bool bIsActiveCollectionEnabled = false;
+
+	TSharedPtr<SSplitter> Splitter;
+
+	TSharedPtr<SVerticalBox>  EditPresetArea;
+	TSharedPtr<SEditableTextBox> EditPresetLabel;
+	TSharedPtr<SEditableTextBox> EditPresetTooltip;
+	TSharedPtr<FToolPresetViewEntry> ActivePresetToEdit;
+
+	TSharedPtr<SPositiveActionButton> AddUserPresetButton;
+	TSharedPtr<SNegativeActionButton> DeleteUserPresetButton;
+};
diff --git a/Engine/Plugins/Experimental/ToolPresets/Source/ToolPresetEditor/Public/IToolPresetEditorModule.h b/Engine/Plugins/Experimental/ToolPresets/Source/ToolPresetEditor/Public/IToolPresetEditorModule.h
new file mode 100644
--- /dev/null
+++ b/Engine/Plugins/Experimental/ToolPresets/Source/ToolPresetEditor/Public/IToolPresetEditorModule.h
@@ -0,0 +1,20 @@
+// Copyright Epic Games, Inc. All Rights Reserved.
+
+#pragma once
+
+#include "CoreMinimal.h"
+#include "Modules/ModuleInterface.h"
+#include "Modules/ModuleManager.h"
+
+TOOLPRESETEDITOR_API extern const FName PresetEditorTabName;
+
+class TOOLPRESETEDITOR_API IToolPresetEditorModule : public IModuleInterface
+{
+public:
+	static IToolPresetEditorModule& Get()
+	{
+		return FModuleManager::LoadModuleChecked<IToolPresetEditorModule>("ToolPresetEditor");
+	}
+
+	virtual void ExecuteOpenPresetEditor() = 0;
+};
diff --git a/Engine/Plugins/Experimental/ToolPresets/Source/ToolPresetEditor/Public/ToolPresetSettings.h b/Engine/Plugins/Experimental/ToolPresets/Source/ToolPresetEditor/Public/ToolPresetSettings.h
new file mode 100644
--- /dev/null
+++ b/Engine/Plugins/Experimental/ToolPresets/Source/ToolPresetEditor/Public/ToolPresetSettings.h
@@ -0,0 +1,65 @@
+// Copyright Epic Games, Inc. All Rights Reserved.
+
+#pragma once
+
+#include "CoreMinimal.h"
+#include "EditorConfigBase.h"
+#include "Engine/DeveloperSettings.h"
+#include "UObject/Object.h"
+#include "UObject/ObjectMacros.h"
+
+#include "ToolPresetSettings.generated.h"
+
+/**
+ * Implements the settings for the PresetEditor.
+ */
+UCLASS(EditorConfig = "UToolPresetUserSettings")
+class TOOLPRESETEDITOR_API UToolPresetUserSettings : public UEditorConfigBase
+{
+	GENERATED_BODY()
+	
+public:
+	UPROPERTY(EditAnywhere, Category = "Interactive Tool Presets|Collections", meta = (EditorConfig, AllowedClasses = "/Script/ToolPresetAsset.InteractiveToolsPresetCollectionAsset"))
+	TSet<FSoftObjectPath> EnabledPresetCollections;
+
+	//~ Ideally the above property would be able to store if the default collection was enabled or not.
+	//~ 
+	//~ However, the default collection itself is stored via an alternative JSON representation, accessed through the PresetAssetSubsystem,
+	//~ to avoid issues with automatic asset generation. Therefore it doesn't have a "path" in the traditional sense that the other collections
+	//~ do, requiring a separate tracking of it's enabled/disabled status.
+
+	UPROPERTY(EditAnywhere, Category = "Interactive Tool Presets|Collections", meta = (EditorConfig))
+	bool bDefaultCollectionEnabled = true;
+
+	static void Initialize();
+	static UToolPresetUserSettings* Get();
+
+private:
+	static TObjectPtr<UToolPresetUserSettings> Instance;
+};
+
+
+/**
+ * Implements the settings for the Tool Project Preset Collections.
+ */
+UCLASS(config = Editor)
+class TOOLPRESETEDITOR_API UToolPresetProjectSettings
+	: public UDeveloperSettings
+{
+	GENERATED_BODY()
+	
+public:
+	// UDeveloperSettings overrides
+
+	virtual FName GetContainerName() const override { return FName("Project"); }
+	virtual FName GetCategoryName() const override { return FName("Plugins"); }
+	virtual FName GetSectionName() const override { return FName("Interactive Tool Presets"); }
+
+	virtual FText GetSectionText() const override { return NSLOCTEXT("ToolPresetSettings", "SectionText", "Interactive Tool Presets"); };
+	virtual FText GetSectionDescription() const override { return NSLOCTEXT("ToolPresetSettings", "SectionDescription", "Manage preset settings at the project level."); };
+
+public:
+	/* Controls which preset collection assets are to be loaded for this project.  */
+	UPROPERTY(config, EditAnywhere, Category = "Interactive Tool Presets|Collections", meta = (AllowedClasses = "/Script/ToolPresetAsset.InteractiveToolsPresetCollectionAsset"))
+	TSet<FSoftObjectPath> LoadedPresetCollections;
+};
diff --git a/Engine/Plugins/Experimental/ToolPresets/Source/ToolPresetEditor/ToolPresetEditor.Build.cs b/Engine/Plugins/Experimental/ToolPresets/Source/ToolPresetEditor/ToolPresetEditor.Build.cs
new file mode 100644
--- /dev/null
+++ b/Engine/Plugins/Experimental/ToolPresets/Source/ToolPresetEditor/ToolPresetEditor.Build.cs
@@ -0,0 +1,38 @@
+// Copyright Epic Games, Inc. All Rights Reserved.
+
+namespace UnrealBuildTool.Rules
+{
+	public class ToolPresetEditor : ModuleRules
+	{
+		public ToolPresetEditor(ReadOnlyTargetRules Target) : base(Target)
+		{
+			PublicDependencyModuleNames.AddRange(
+				new string[] {
+					"Core",					
+				}
+			);
+
+			PrivateDependencyModuleNames.AddRange(
+				new string[] {
+					"ApplicationCore",
+					"AssetTools",
+					"ContentBrowser",
+					"ContentBrowserData",
+					"CoreUObject",
+					"DeveloperSettings",
+					"EditorConfig",
+					"EditorFramework",
+					"EditorStyle",
+					"Engine",
+					"InputCore",
+					"Projects",
+					"Slate",
+					"SlateCore",
+					"ToolPresetAsset",
+					"ToolWidgets",
+					"UnrealEd",
+				}
+			);
+		}
+	}
+}
diff --git a/Engine/Plugins/Experimental/ToolPresets/ToolPresets.uplugin b/Engine/Plugins/Experimental/ToolPresets/ToolPresets.uplugin
--- a/Engine/Plugins/Experimental/ToolPresets/ToolPresets.uplugin
+++ b/Engine/Plugins/Experimental/ToolPresets/ToolPresets.uplugin
@@ -17,12 +17,12 @@
 	"IsExperimentalVersion": true,
 	"Modules": [
 		{
-			"Name": "PresetAsset",
+			"Name": "ToolPresetAsset",
 			"Type": "Editor",
 			"LoadingPhase": "Default"
 		},
 		{
-			"Name": "PresetEditor",
+			"Name": "ToolPresetEditor",
 			"Type": "Editor",
 			"LoadingPhase": "Default"
 		}
