[ToolPresets] module name cis fix

diff --git a/Engine/Plugins/Experimental/ToolPresets/Source/ToolPresetAsset/Private/ToolPresetAssetModule.cpp b/Engine/Plugins/Experimental/ToolPresets/Source/ToolPresetAsset/Private/ToolPresetAssetModule.cpp
--- a/Engine/Plugins/Experimental/ToolPresets/Source/ToolPresetAsset/Private/ToolPresetAssetModule.cpp
+++ b/Engine/Plugins/Experimental/ToolPresets/Source/ToolPresetAsset/Private/ToolPresetAssetModule.cpp
@@ -61,6 +61,6 @@ void FToolPresetAssetModule::OnPostEngineInit()
 }
 
 
-IMPLEMENT_MODULE(FToolPresetAssetModule, PresetAsset);
+IMPLEMENT_MODULE(FToolPresetAssetModule, ToolPresetAsset);
 
 #undef LOCTEXT_NAMESPACE
