ModelingMode: Fix crash issue with the Interactive Tool Presets where non-property set data could get loaded and cause internal tool data structures to become corrupted.

diff --git a/Engine/Plugins/Editor/ModelingToolsEditorMode/Source/ModelingToolsEditorMode/Private/ModelingToolsEditorModeToolkit.cpp b/Engine/Plugins/Editor/ModelingToolsEditorMode/Source/ModelingToolsEditorMode/Private/ModelingToolsEditorModeToolkit.cpp
--- a/Engine/Plugins/Editor/ModelingToolsEditorMode/Source/ModelingToolsEditorMode/Private/ModelingToolsEditorModeToolkit.cpp
+++ b/Engine/Plugins/Editor/ModelingToolsEditorMode/Source/ModelingToolsEditorMode/Private/ModelingToolsEditorModeToolkit.cpp
@@ -1291,6 +1291,13 @@ void FModelingToolsEditorModeToolkit::LoadPresetFromCollection(const int32 Prese
 	FModelingToolsEditorModeToolkitLocals::ExecuteWithPresetAndTool(*OwningEditorMode, EToolSide::Left, CollectionPath,
 	[this, PresetIndex](UInteractiveToolsPresetCollectionAsset& Preset, UInteractiveTool& Tool) {
 		TArray<UObject*> PropertySets = Tool.GetToolProperties();
+
+		// We only want to load the properties that are actual property sets, since the tool might have added other types of objects we don't
+		// want to deserialize.
+		PropertySets.RemoveAll([this](UObject* Object) {
+			return Cast<UInteractiveToolPropertySet>(Object) == nullptr;
+		});
+
 		if (!Preset.PerToolPresets.Contains(Tool.GetClass()->GetName()))			
 		{
 			return;
diff --git a/Engine/Plugins/Experimental/ToolPresets/Source/ToolPresetAsset/Private/ToolPresetAsset.cpp b/Engine/Plugins/Experimental/ToolPresets/Source/ToolPresetAsset/Private/ToolPresetAsset.cpp
--- a/Engine/Plugins/Experimental/ToolPresets/Source/ToolPresetAsset/Private/ToolPresetAsset.cpp
+++ b/Engine/Plugins/Experimental/ToolPresets/Source/ToolPresetAsset/Private/ToolPresetAsset.cpp
@@ -52,15 +52,22 @@ void FInteractiveToolPresetDefinition::LoadStoredPropertyData(TArray<UObject*>&
 		{
 			for (FProperty* Prop : TFieldRange<FProperty>(PropertySet->GetClass()))
 			{
-				FString FieldName;
-				Prop->GetName(FieldName);
-				TSharedPtr<FJsonValue>* JsonField = PropertyJsonObject->Values.Find(FieldName);
-				if (JsonField)
+#if WITH_EDITOR
+				if (!(Prop->HasMetaData(TEXT("TransientToolProperty")) ||
+					Prop->GetPropertyFlags() & EPropertyFlags::CPF_SkipSerialization ||
+					Prop->GetPropertyFlags() & EPropertyFlags::CPF_Transient))
+#endif
 				{
-					if (FJsonObjectConverter::JsonValueToUProperty(*JsonField, Prop, Prop->ContainerPtrToValuePtr<void>(PropertySet)))
+					FString FieldName;
+					Prop->GetName(FieldName);
+					TSharedPtr<FJsonValue>* JsonField = PropertyJsonObject->Values.Find(FieldName);
+					if (JsonField)
 					{
-						FPropertyChangedEvent ChangeEvent(Prop, EPropertyChangeType::ValueSet);
-						PropertySet->PostEditChangeProperty(ChangeEvent);
+						if (FJsonObjectConverter::JsonValueToUProperty(*JsonField, Prop, Prop->ContainerPtrToValuePtr<void>(PropertySet)))
+						{
+							FPropertyChangedEvent ChangeEvent(Prop, EPropertyChangeType::ValueSet);
+							PropertySet->PostEditChangeProperty(ChangeEvent);
+						}
 					}
 				}
 			}
