ToolPresets: Fixed issue where newly created preset collections from the Content Browser would start with no label, causing confusion in the Preset Manager.

diff --git a/Engine/Plugins/Experimental/ToolPresets/Source/ToolPresetAsset/Public/ToolPresetAsset.h b/Engine/Plugins/Experimental/ToolPresets/Source/ToolPresetAsset/Public/ToolPresetAsset.h
--- a/Engine/Plugins/Experimental/ToolPresets/Source/ToolPresetAsset/Public/ToolPresetAsset.h
+++ b/Engine/Plugins/Experimental/ToolPresets/Source/ToolPresetAsset/Public/ToolPresetAsset.h
@@ -105,6 +105,7 @@ public:
 		{
 			ensure(0 != (RF_Public & Flags));
 			Preset = NewObject<UInteractiveToolsPresetCollectionAsset>(InParent, Class, Name, Flags);
+			Preset->CollectionLabel = FText::FromName(Name);
 		}
 		return Preset;
 	}
