ToolPresets: Fix issue where renaming a user collection in the preset manager with spaces in the name would fail to rename due to asset naming requirements.

diff --git a/Engine/Plugins/Experimental/ToolPresets/Source/ToolPresetEditor/Private/Widgets/SToolPresetManager.cpp b/Engine/Plugins/Experimental/ToolPresets/Source/ToolPresetEditor/Private/Widgets/SToolPresetManager.cpp
--- a/Engine/Plugins/Experimental/ToolPresets/Source/ToolPresetEditor/Private/Widgets/SToolPresetManager.cpp
+++ b/Engine/Plugins/Experimental/ToolPresets/Source/ToolPresetEditor/Private/Widgets/SToolPresetManager.cpp
@@ -35,6 +35,7 @@
 #include "Widgets/Views/SListView.h"
 #include "Widgets/Views/STableRow.h"
 #include "Widgets/Views/STreeView.h"
+#include "PackageTools.h"
 
 #define LOCTEXT_NAMESPACE "SToolPresetManager"
 
@@ -1281,9 +1282,10 @@ void SToolPresetManager::CollectionRenameEnded(TSharedPtr<FToolPresetViewEntry>
 		}
 
 		TArray< FAssetRenameData> RenameData;
-
+		
 		FString NewPackageName, NewAssetName;
-		IAssetTools::Get().CreateUniqueAssetName(NewText.ToString(), "", NewPackageName, NewAssetName);
+		FString SanitizedName = UPackageTools::SanitizePackageName(NewText.ToString());
+		IAssetTools::Get().CreateUniqueAssetName(SanitizedName, "", NewPackageName, NewAssetName);
 
 		RenameData.SetNum(1);		
 		RenameData[0].Asset = CollectionObject;
