[ContentBrowser] New Add Menu Misc Menu

#Synced-CL 52619962

diff --git a/Engine/Plugins/Animation/LiveLink/Source/LiveLinkEditor/LiveLinkEditor.Build.cs b/Engine/Plugins/Animation/LiveLink/Source/LiveLinkEditor/LiveLinkEditor.Build.cs
--- a/Engine/Plugins/Animation/LiveLink/Source/LiveLinkEditor/LiveLinkEditor.Build.cs
+++ b/Engine/Plugins/Animation/LiveLink/Source/LiveLinkEditor/LiveLinkEditor.Build.cs
@@ -19,6 +19,7 @@ namespace UnrealBuildTool.Rules
 			PrivateDependencyModuleNames.AddRange(
 				new string[]
 				{
+					"AssetDefinition",
 					"AnimGraph",
 					"AssetRegistry",
 					"BlueprintGraph",
@@ -29,6 +30,7 @@ namespace UnrealBuildTool.Rules
 					"EditorStyle",
 					"EditorWidgets",
 					"Engine",
+					"EngineAssetDefinitions",
 					"InputCore",
 					"KismetCompiler",
 					"GraphEditor",
diff --git a/Engine/Plugins/Animation/LiveLink/Source/LiveLinkEditor/Private/AssetDefinition/AssetDefinition_LiveLinkBlueprintVirtualSubject.h b/Engine/Plugins/Animation/LiveLink/Source/LiveLinkEditor/Private/AssetDefinition/AssetDefinition_LiveLinkBlueprintVirtualSubject.h
new file mode 100644
--- /dev/null
+++ b/Engine/Plugins/Animation/LiveLink/Source/LiveLinkEditor/Private/AssetDefinition/AssetDefinition_LiveLinkBlueprintVirtualSubject.h
@@ -0,0 +1,25 @@
+// Copyright Epic Games, Inc. All Rights Reserved.
+
+#pragma once
+#include "Script/AssetDefinition_Blueprint.h"
+#include "VirtualSubjects/LiveLinkBlueprintVirtualSubject.h"
+#include "AssetDefinition_LiveLinkBlueprintVirtualSubject.generated.h"
+
+UCLASS()
+class UAssetDefinition_LiveLinkBlueprintVirtualSubject : public UAssetDefinition_Blueprint
+{
+	GENERATED_BODY()
+public:
+	// UAssetDefinition Begin
+	virtual FText GetAssetDisplayName() const override { return NSLOCTEXT("AssetDefinition", "AssetDefinition_LiveLinkBlueprintVirtualSubject", "Blueprint Virtual Subject"); }
+	virtual TSoftClassPtr<UObject> GetAssetClass() const override { return ULiveLinkBlueprintVirtualSubject::StaticClass(); }
+	virtual TConstArrayView<FAssetCategoryPath> GetAssetCategories() const override
+	{
+		static const TArray<FAssetCategoryPath> Categories = 
+			{
+				FAssetCategoryPath(EAssetCategoryPaths::Misc, NSLOCTEXT("AssetDefinition", "AssetDefinition_LiveLinkBlueprintVirtualSubjectSubMenu", "Other"), ECategoryMenuType::Section)
+			};
+		return Categories;
+	}
+	// UAssetDefinition End
+};
diff --git a/Engine/Plugins/Animation/LiveLink/Source/LiveLinkEditor/Private/LiveLinkEditorModule.cpp b/Engine/Plugins/Animation/LiveLink/Source/LiveLinkEditor/Private/LiveLinkEditorModule.cpp
--- a/Engine/Plugins/Animation/LiveLink/Source/LiveLinkEditor/Private/LiveLinkEditorModule.cpp
+++ b/Engine/Plugins/Animation/LiveLink/Source/LiveLinkEditor/Private/LiveLinkEditorModule.cpp
@@ -1,5 +1,6 @@
 // Copyright Epic Games, Inc. All Rights Reserved.
 
+#include "AssetToolsModule.h"
 #include "Interfaces/IPluginManager.h"
 
 #include "Brushes/SlateBoxBrush.h"
@@ -101,6 +102,9 @@ public:
 
 		FLiveLinkClientCommands::Register();
 
+		IAssetTools& AssetTools = FModuleManager::LoadModuleChecked<FAssetToolsModule>("AssetTools").Get();
+		AssetTools.RegisterAdvancedAssetCategory("LiveLink", LOCTEXT("AssetCategoryName", "Live Link"));
+
 		const FVector2D Icon8x8(8.0f, 8.0f);
 		const FVector2D Icon16x16(16.0f, 16.0f);
 		const FVector2D Icon40x40(40.0f, 40.0f);
diff --git a/Engine/Plugins/Animation/LiveLink/Source/LiveLinkEditor/Private/VirtualSubjects/LiveLinkBlueprintVirtualSubjectFactory.cpp b/Engine/Plugins/Animation/LiveLink/Source/LiveLinkEditor/Private/VirtualSubjects/LiveLinkBlueprintVirtualSubjectFactory.cpp
--- a/Engine/Plugins/Animation/LiveLink/Source/LiveLinkEditor/Private/VirtualSubjects/LiveLinkBlueprintVirtualSubjectFactory.cpp
+++ b/Engine/Plugins/Animation/LiveLink/Source/LiveLinkEditor/Private/VirtualSubjects/LiveLinkBlueprintVirtualSubjectFactory.cpp
@@ -245,8 +245,16 @@ bool ULiveLinkBlueprintVirtualSubjectFactory::ShouldShowInNewMenu() const
 
 uint32 ULiveLinkBlueprintVirtualSubjectFactory::GetMenuCategories() const
 {
-	IAssetTools& AssetTools = FModuleManager::LoadModuleChecked<FAssetToolsModule>("AssetTools").Get();
-	return AssetTools.RegisterAdvancedAssetCategory("LiveLink", LOCTEXT("AssetCategoryName", "Live Link"));
+	return EAssetTypeCategories::Misc;
+}
+
+TArray<FAssetCategoryPath> ULiveLinkBlueprintVirtualSubjectFactory::GetAssetMenuPathsForCategory(FName InCategory) const
+{
+	if (ParentClass)
+	{
+		return GetAssetMenuPaths(ParentClass, InCategory);
+	}
+	return Super::GetAssetMenuPathsForCategory(InCategory);
 }
 
 #undef  LOCTEXT_NAMESPACE
diff --git a/Engine/Plugins/Animation/LiveLink/Source/LiveLinkEditor/Public/VirtualSubjects/LiveLinkBlueprintVirtualSubjectFactory.h b/Engine/Plugins/Animation/LiveLink/Source/LiveLinkEditor/Public/VirtualSubjects/LiveLinkBlueprintVirtualSubjectFactory.h
--- a/Engine/Plugins/Animation/LiveLink/Source/LiveLinkEditor/Public/VirtualSubjects/LiveLinkBlueprintVirtualSubjectFactory.h
+++ b/Engine/Plugins/Animation/LiveLink/Source/LiveLinkEditor/Public/VirtualSubjects/LiveLinkBlueprintVirtualSubjectFactory.h
@@ -27,6 +27,7 @@ public:
 	UE_API virtual UObject* FactoryCreateNew(UClass* Class, UObject* InParent, FName Name, EObjectFlags Flags, UObject* Context, FFeedbackContext* Warn, FName CallingContext) override;
 	UE_API virtual bool ShouldShowInNewMenu() const override;
 	UE_API virtual uint32 GetMenuCategories() const override;
+	UE_API virtual TArray<FAssetCategoryPath> GetAssetMenuPathsForCategory(FName InCategory) const override;
 	//~ Begin UFactory Interface
 };
 
diff --git a/Engine/Plugins/Editor/ConsoleVariablesEditor/Source/ConsoleVariablesEditor/Private/Factories/ConsoleVariablesEditorFactory.cpp b/Engine/Plugins/Editor/ConsoleVariablesEditor/Source/ConsoleVariablesEditor/Private/Factories/ConsoleVariablesEditorFactory.cpp
--- a/Engine/Plugins/Editor/ConsoleVariablesEditor/Source/ConsoleVariablesEditor/Private/Factories/ConsoleVariablesEditorFactory.cpp
+++ b/Engine/Plugins/Editor/ConsoleVariablesEditor/Source/ConsoleVariablesEditor/Private/Factories/ConsoleVariablesEditorFactory.cpp
@@ -19,22 +19,16 @@ UConsoleVariablesEditorFactory::UConsoleVariablesEditorFactory()
 
 	// This factory will open the editor for each new object.
 	bEditAfterNew = true;
-};
+}
 
 UObject* UConsoleVariablesEditorFactory::FactoryCreateNew(UClass* InClass, UObject* InParent, FName InName, EObjectFlags Flags, UObject* Context, FFeedbackContext* Warn)
 {
 	return NewObject<UConsoleVariablesAsset>(InParent, InClass, InName, Flags);
-};
-
+}
 
 bool UConsoleVariablesEditorFactory::ShouldShowInNewMenu() const
 {
 	return true;
 }
-uint32 UConsoleVariablesEditorFactory::GetMenuCategories() const
-{
-	IAssetTools& AssetTools = FModuleManager::LoadModuleChecked<FAssetToolsModule>("AssetTools").Get();
-	return AssetTools.RegisterAdvancedAssetCategory("Console Variables Editor", LOCTEXT("AssetCategoryName", "Console Variables Editor"));
-}
 
 #undef LOCTEXT_NAMESPACE
diff --git a/Engine/Plugins/Editor/ConsoleVariablesEditor/Source/ConsoleVariablesEditor/Private/Factories/ConsoleVariablesEditorFactory.h b/Engine/Plugins/Editor/ConsoleVariablesEditor/Source/ConsoleVariablesEditor/Private/Factories/ConsoleVariablesEditorFactory.h
--- a/Engine/Plugins/Editor/ConsoleVariablesEditor/Source/ConsoleVariablesEditor/Private/Factories/ConsoleVariablesEditorFactory.h
+++ b/Engine/Plugins/Editor/ConsoleVariablesEditor/Source/ConsoleVariablesEditor/Private/Factories/ConsoleVariablesEditorFactory.h
@@ -18,5 +18,4 @@ public:
 	
 	virtual UObject* FactoryCreateNew(UClass* InClass, UObject* InParent, FName InName, EObjectFlags Flags, UObject* Context, FFeedbackContext* Warn) override;
 	virtual bool ShouldShowInNewMenu() const override;
-	virtual uint32 GetMenuCategories() const override;
 };
\ No newline at end of file
diff --git a/Engine/Plugins/Editor/EngineAssetDefinitions/Source/Private/AssetDefinition_LandscapeLayer.h b/Engine/Plugins/Editor/EngineAssetDefinitions/Source/Private/AssetDefinition_LandscapeLayer.h
--- a/Engine/Plugins/Editor/EngineAssetDefinitions/Source/Private/AssetDefinition_LandscapeLayer.h
+++ b/Engine/Plugins/Editor/EngineAssetDefinitions/Source/Private/AssetDefinition_LandscapeLayer.h
@@ -17,5 +17,13 @@ public:
 	virtual FText GetAssetDisplayName() const override { return NSLOCTEXT("AssetTypeActions", "AssetTypeActions_LandscapeLayer", "Landscape Layer"); }
 	virtual FLinearColor GetAssetColor() const override { return FLinearColor(FColor(128,192,255)); }
 	virtual TSoftClassPtr<UObject> GetAssetClass() const override { return ULandscapeLayerInfoObject::StaticClass(); }
+	virtual TConstArrayView<FAssetCategoryPath> GetAssetCategories() const override
+	{
+		static const TArray<FAssetCategoryPath> Categories = 
+			{
+				FAssetCategoryPath(EAssetCategoryPaths::Misc, NSLOCTEXT("AssetDefinition", "LandscapeLayer_SubMenu", "Other"), ECategoryMenuType::Section)
+			};
+		return Categories;
+	}
 	// UAssetDefinition End
 };
diff --git a/Engine/Plugins/Editor/EngineAssetDefinitions/Source/Private/AssetDefinition_SubUVAnimation.h b/Engine/Plugins/Editor/EngineAssetDefinitions/Source/Private/AssetDefinition_SubUVAnimation.h
--- a/Engine/Plugins/Editor/EngineAssetDefinitions/Source/Private/AssetDefinition_SubUVAnimation.h
+++ b/Engine/Plugins/Editor/EngineAssetDefinitions/Source/Private/AssetDefinition_SubUVAnimation.h
@@ -17,5 +17,13 @@ public:
 	virtual FText GetAssetDisplayName() const override { return NSLOCTEXT("AssetTypeActions", "AssetTypeActions_SubUVAnimation", "Sub UV Animation"); }
 	virtual FLinearColor GetAssetColor() const override { return FLinearColor(FColor(255,255,255)); }
 	virtual TSoftClassPtr<UObject> GetAssetClass() const override { return USubUVAnimation::StaticClass(); }
+	virtual TConstArrayView<FAssetCategoryPath> GetAssetCategories() const override
+	{
+		static const TArray<FAssetCategoryPath> Categories = 
+			{
+				FAssetCategoryPath(EAssetCategoryPaths::Misc, NSLOCTEXT("AssetDefinition", "SubUVAnimation_SubMenu", "Other"), ECategoryMenuType::Section)
+			};
+		return Categories;
+	}
 	// UAssetDefinition End
 };
diff --git a/Engine/Plugins/Editor/ObjectMixer/ObjectMixer/ObjectMixer.uplugin b/Engine/Plugins/Editor/ObjectMixer/ObjectMixer/ObjectMixer.uplugin
--- a/Engine/Plugins/Editor/ObjectMixer/ObjectMixer/ObjectMixer.uplugin
+++ b/Engine/Plugins/Editor/ObjectMixer/ObjectMixer/ObjectMixer.uplugin
@@ -21,5 +21,11 @@
 			"Type": "Editor",
 			"LoadingPhase": "Default"
 		}
+	],
+	"Plugins": [
+		{
+			"Name": "EngineAssetDefinitions",
+			"Enabled": true
+		}
 	]
 }
diff --git a/Engine/Plugins/Editor/ObjectMixer/ObjectMixer/Source/ObjectMixer/ObjectMixerEditor.build.cs b/Engine/Plugins/Editor/ObjectMixer/ObjectMixer/Source/ObjectMixer/ObjectMixerEditor.build.cs
--- a/Engine/Plugins/Editor/ObjectMixer/ObjectMixer/Source/ObjectMixer/ObjectMixerEditor.build.cs
+++ b/Engine/Plugins/Editor/ObjectMixer/ObjectMixer/Source/ObjectMixer/ObjectMixerEditor.build.cs
@@ -29,6 +29,7 @@ public class ObjectMixerEditor : ModuleRules
 		PrivateDependencyModuleNames.AddRange(
 			new string[]
 			{
+				"AssetDefinition",
 				"AdvancedWidgets",
 				"ApplicationCore",
 				"AssetRegistry",
@@ -40,6 +41,7 @@ public class ObjectMixerEditor : ModuleRules
 				"EditorConfig",
 				"EditorStyle",
 				"EditorWidgets",
+				"EngineAssetDefinitions",
 				"InputCore",
 				"Kismet",
 				"LevelEditor",
@@ -55,7 +57,7 @@ public class ObjectMixerEditor : ModuleRules
 				"UMG",
 				"UnrealEd",
 				"WorkspaceMenuStructure",
-				"SourceControl"
+				"SourceControl",
 			}
 		);
 	}
diff --git a/Engine/Plugins/Editor/ObjectMixer/ObjectMixer/Source/ObjectMixer/Private/AssetDefinition/AssetDefinition_MixerBlueprintObjectFilter.h b/Engine/Plugins/Editor/ObjectMixer/ObjectMixer/Source/ObjectMixer/Private/AssetDefinition/AssetDefinition_MixerBlueprintObjectFilter.h
new file mode 100644
--- /dev/null
+++ b/Engine/Plugins/Editor/ObjectMixer/ObjectMixer/Source/ObjectMixer/Private/AssetDefinition/AssetDefinition_MixerBlueprintObjectFilter.h
@@ -0,0 +1,25 @@
+// Copyright Epic Games, Inc. All Rights Reserved.
+
+#pragma once
+#include "ObjectFilter/ObjectMixerEditorObjectFilter.h"
+#include "Script/AssetDefinition_Blueprint.h"
+#include "AssetDefinition_MixerBlueprintObjectFilter.generated.h"
+
+UCLASS()
+class UAssetDefinition_MixerBlueprintObjectFilter : public UAssetDefinition_Blueprint
+{
+	GENERATED_BODY()
+public:
+	// UAssetDefinition Begin
+	virtual FText GetAssetDisplayName() const override { return NSLOCTEXT("AssetDefinition", "AssetDefinition_MixerBlueprintObjectFilter", "Object Mixer Filter"); }
+	virtual TSoftClassPtr<UObject> GetAssetClass() const override { return UObjectMixerBlueprintObjectFilter::StaticClass(); }
+	virtual TConstArrayView<FAssetCategoryPath> GetAssetCategories() const override
+	{
+		static const TArray<FAssetCategoryPath> Categories = 
+			{
+				FAssetCategoryPath(EAssetCategoryPaths::Misc, NSLOCTEXT("AssetDefinition", "AssetDefinition_MixerBlueprintObjectFilterSubMenu", "Other"), ECategoryMenuType::Section)
+			};
+		return Categories;
+	}
+	// UAssetDefinition End
+};
diff --git a/Engine/Plugins/Editor/ObjectMixer/ObjectMixer/Source/ObjectMixer/Private/ObjectMixerFilterFactory.cpp b/Engine/Plugins/Editor/ObjectMixer/ObjectMixer/Source/ObjectMixer/Private/ObjectMixerFilterFactory.cpp
--- a/Engine/Plugins/Editor/ObjectMixer/ObjectMixer/Source/ObjectMixer/Private/ObjectMixerFilterFactory.cpp
+++ b/Engine/Plugins/Editor/ObjectMixer/ObjectMixer/Source/ObjectMixer/Private/ObjectMixerFilterFactory.cpp
@@ -58,4 +58,13 @@ FString UObjectMixerBlueprintFilterFactory::GetToolTipDocumentationExcerpt() con
 	return TEXT("UObjectMixerBlueprintFilter");
 }
 
+TArray<FAssetCategoryPath> UObjectMixerBlueprintFilterFactory::GetAssetMenuPathsForCategory(FName InCategory) const
+{
+	if (ParentClass)
+	{
+		return GetAssetMenuPaths(ParentClass, InCategory);
+	}
+	return Super::GetAssetMenuPathsForCategory(InCategory);
+}
+
 #undef LOCTEXT_NAMESPACE
diff --git a/Engine/Plugins/Editor/ObjectMixer/ObjectMixer/Source/ObjectMixer/Public/ObjectMixerFilterFactory.h b/Engine/Plugins/Editor/ObjectMixer/ObjectMixer/Source/ObjectMixer/Public/ObjectMixerFilterFactory.h
--- a/Engine/Plugins/Editor/ObjectMixer/ObjectMixer/Source/ObjectMixer/Public/ObjectMixerFilterFactory.h
+++ b/Engine/Plugins/Editor/ObjectMixer/ObjectMixer/Source/ObjectMixer/Public/ObjectMixerFilterFactory.h
@@ -26,6 +26,7 @@ public:
 	UE_API virtual uint32 GetMenuCategories() const override;
 	UE_API virtual FText GetToolTip() const override;
 	UE_API virtual FString GetToolTipDocumentationExcerpt() const override;
+	UE_API virtual TArray<FAssetCategoryPath> GetAssetMenuPathsForCategory(FName InCategory) const override;
 	//~ Begin UFactory Interface
 	
 };
diff --git a/Engine/Plugins/Experimental/Animation/HierarchyTable/Source/Editor/Private/HierarchyTableAssetDefinition.cpp b/Engine/Plugins/Experimental/Animation/HierarchyTable/Source/Editor/Private/HierarchyTableAssetDefinition.cpp
--- a/Engine/Plugins/Experimental/Animation/HierarchyTable/Source/Editor/Private/HierarchyTableAssetDefinition.cpp
+++ b/Engine/Plugins/Experimental/Animation/HierarchyTable/Source/Editor/Private/HierarchyTableAssetDefinition.cpp
@@ -36,7 +36,10 @@ EAssetCommandResult UAssetDefinition_HierarchyTable::OpenAssets(const FAssetOpen
 
 TConstArrayView<FAssetCategoryPath> UAssetDefinition_HierarchyTable::GetAssetCategories() const
 {
-	static const auto Categories = { FAssetCategoryPath(EAssetCategoryPaths::Misc) };
+	static const TArray<FAssetCategoryPath> Categories = 
+		{
+			FAssetCategoryPath(EAssetCategoryPaths::Misc, LOCTEXT("AssetDefinition_HierarchyTableSubMenu", "Other"), ECategoryMenuType::Section)
+		};
 	return Categories;
 }
 
diff --git a/Engine/Plugins/Experimental/ToolPresets/Source/ToolPresetAsset/Private/AssetDefinition_ToolPresetAsset.h b/Engine/Plugins/Experimental/ToolPresets/Source/ToolPresetAsset/Private/AssetDefinition_ToolPresetAsset.h
--- a/Engine/Plugins/Experimental/ToolPresets/Source/ToolPresetAsset/Private/AssetDefinition_ToolPresetAsset.h
+++ b/Engine/Plugins/Experimental/ToolPresets/Source/ToolPresetAsset/Private/AssetDefinition_ToolPresetAsset.h
@@ -18,12 +18,16 @@ class UAssetDefinition_InteractiveToolsPresetCollectionAsset : public UAssetDefi
 
 public:
 	// UAssetDefinition Begin
-	virtual FText GetAssetDisplayName() const override { return NSLOCTEXT("AssetTypeActions", "AssetTypeActions_InteractiveToolsPresetCollectionAsset", "Interactive Tools Preset Collection"); }
+	virtual FText GetAssetDisplayName() const override { return NSLOCTEXT("AssetTypeActions", "AssetTypeActions_InteractiveToolsPresetCollectionAsset", "Tool Preset"); }
 	virtual FLinearColor GetAssetColor() const override { return FLinearColor(FColor(175, 0, 128)); }
 	virtual TSoftClassPtr<UObject> GetAssetClass() const override { return UInteractiveToolsPresetCollectionAsset::StaticClass(); }
 	virtual TConstArrayView<FAssetCategoryPath> GetAssetCategories() const override
 	{
-		return TConstArrayView<FAssetCategoryPath>();
+		static const TArray<FAssetCategoryPath> Categories = 
+			{
+				FAssetCategoryPath(EAssetCategoryPaths::Misc, NSLOCTEXT("AssetDefinition", "AssetDefinition_InteractiveToolsPresetCollectionAssetSubMenu", "Other"), ECategoryMenuType::Section)
+			};
+		return Categories;
 	}
 	virtual FText GetObjectDisplayNameText(UObject* Object) const override { return FText::FromString(TEXT("UInteractiveToolsPresetCollectionAsset")); }
 	// UAssetDefinition End
diff --git a/Engine/Plugins/Experimental/ToolPresets/Source/ToolPresetAsset/Public/ToolPresetAsset.h b/Engine/Plugins/Experimental/ToolPresets/Source/ToolPresetAsset/Public/ToolPresetAsset.h
--- a/Engine/Plugins/Experimental/ToolPresets/Source/ToolPresetAsset/Public/ToolPresetAsset.h
+++ b/Engine/Plugins/Experimental/ToolPresets/Source/ToolPresetAsset/Public/ToolPresetAsset.h
@@ -90,17 +90,12 @@ public:
 		bEditAfterNew = true;
 	}
 
-	FText GetDisplayName() const override
-	{
-		return LOCTEXT("DisplayName", "Tool Preset");
-	}
-
-	FText GetToolTip() const override
+	virtual FText GetToolTip() const override
 	{
 		return LOCTEXT("Tooltip", "Tool Presets capture the state of tool settings for later reloading.");
 	}
 
-	UObject* FactoryCreateNew(UClass* Class, UObject* InParent, FName Name, EObjectFlags Flags, UObject* Context, FFeedbackContext* Warn) override
+	virtual UObject* FactoryCreateNew(UClass* Class, UObject* InParent, FName Name, EObjectFlags Flags, UObject* Context, FFeedbackContext* Warn) override
 	{
 		UInteractiveToolsPresetCollectionAsset* Preset = nullptr;
 		if (ensure(SupportedClass == Class))
@@ -112,18 +107,12 @@ public:
 		return Preset;
 	}
 
-	uint32 GetMenuCategories() const override
-	{
-		IAssetTools& AssetTools = FModuleManager::LoadModuleChecked<FAssetToolsModule>("AssetTools").Get();
-		return AssetTools.RegisterAdvancedAssetCategory("ToolPresets", LOCTEXT("AssetCategoryName", "Tool Presets"));
-	}
-
-	FString GetDefaultNewAssetName() const override
+	virtual FString GetDefaultNewAssetName() const override
 	{
 		return TEXT("Tool Preset");
 	}
 
-	bool ShouldShowInNewMenu() const override { return true; }
+	virtual bool ShouldShowInNewMenu() const override { return true; }
 };
 
 #undef LOCTEXT_NAMESPACE
